HLODOutliner changes

diff --git a/Engine/Source/Editor/EditorStyle/Private/SlateEditorStyle.cpp b/Engine/Source/Editor/EditorStyle/Private/SlateEditorStyle.cpp
--- a/Engine/Source/Editor/EditorStyle/Private/SlateEditorStyle.cpp
+++ b/Engine/Source/Editor/EditorStyle/Private/SlateEditorStyle.cpp
@@ -4257,6 +4257,7 @@ void FSlateEditorStyle::FStyle::SetupLevelEditorStyle()
 			Set( "LevelEditor.Tabs.StatsViewer", new IMAGE_BRUSH( "/Icons/icon_tab_Stats_16x", Icon16x16 ) );
 			Set( "LevelEditor.Tabs.Toolbar", new IMAGE_BRUSH( "/Icons/icon_tab_Toolbars_16x", Icon16x16 ) );
 			Set( "LevelEditor.Tabs.Viewports", new IMAGE_BRUSH( "/Icons/icon_tab_Viewports_16x", Icon16x16 ) );
+			Set( "LevelEditor.Tabs.HLOD", new IMAGE_BRUSH("/Icons/icon_tab_layers_16px", Icon16x16));
 		}
 #endif
 
diff --git a/Engine/Source/Editor/HierarchicalLODOutliner/HierarchicalLODOutliner.Build.cs b/Engine/Source/Editor/HierarchicalLODOutliner/HierarchicalLODOutliner.Build.cs
new file mode 100644
--- /dev/null
+++ b/Engine/Source/Editor/HierarchicalLODOutliner/HierarchicalLODOutliner.Build.cs
@@ -0,0 +1,46 @@
+// Copyright 1998-2015 Epic Games, Inc. All Rights Reserved.
+
+using UnrealBuildTool;
+
+public class HierarchicalLODOutliner : ModuleRules
+{
+    public HierarchicalLODOutliner(TargetInfo Target)
+	{
+		PublicDependencyModuleNames.AddRange
+        (
+            new string[] {
+				"Core",
+                "InputCore",
+                "CoreUObject",
+				"RHI",
+				"RenderCore",
+				"Slate",
+                "EditorStyle",
+                "Engine",
+                "UnrealEd",
+                "PropertyEditor"
+			}
+        );
+
+
+        PrivateDependencyModuleNames.AddRange(
+             new string[] {
+					"Engine",
+                    "UnrealEd"
+				}
+         );
+        
+        PrivateDependencyModuleNames.AddRange(
+            new string[] {
+				"SlateCore",
+			}
+        );
+
+        PrivateIncludePathModuleNames.AddRange(
+        new string[] {
+				"Messaging",
+				"SessionServices",
+			}
+    );
+	}
+}
diff --git a/Engine/Source/Editor/HierarchicalLODOutliner/Private/HLODOutliner.cpp b/Engine/Source/Editor/HierarchicalLODOutliner/Private/HLODOutliner.cpp
new file mode 100644
--- /dev/null
+++ b/Engine/Source/Editor/HierarchicalLODOutliner/Private/HLODOutliner.cpp
@@ -0,0 +1,1248 @@
+// Copyright 1998-2015 Epic Games, Inc. All Rights Reserved.
+
+#include "HierarchicalLODOutlinerPrivatePCH.h"
+#include "HLODOutliner.h"
+
+#include "Engine.h"
+#include "Engine/LodActor.h"
+#include "Engine/World.h"
+#include "HierarchicalLOD.h"
+
+#include "Editor.h"
+#include "EditorModeManager.h"
+#include "PropertyEditorModule.h"
+#include "IDetailsView.h"
+#include "SListView.h"
+#include "ScopedTransaction.h"
+
+#include "ITreeItem.h"
+#include "LODActorItem.h"
+#include "LODLevelItem.h"
+#include "StaticMeshActorItem.h"
+#include "HLODTreeWidgetitem.h"
+#include "HLODSelectionActor.h"
+
+#include "HierarchicalLODUtils.h"
+
+#define LOCTEXT_NAMESPACE "HLODOutliner"
+
+namespace HLODOutliner
+{
+	SHLODOutliner::SHLODOutliner()
+	{
+		bNeedsRefresh = true;
+		CurrentWorld = nullptr;
+		ForcedLODLevel = -1;
+	}
+
+	SHLODOutliner::~SHLODOutliner()
+	{
+		DeregisterDelegates();	
+		DestroySelectionActors();
+		CurrentWorld = nullptr;
+		HLODTreeRoot.Empty();
+		SelectedNodes.Empty();		
+		AllNodes.Empty();
+		SelectionActors.Empty();
+		LODLevelBuildFlags.Empty();
+		LODLevelActors.Empty();
+		PendingActions.Empty();
+	}
+
+	BEGIN_SLATE_FUNCTION_BUILD_OPTIMIZATION
+		void SHLODOutliner::Construct(const FArguments& InArgs)
+	{
+		CreateSettingsView();
+
+		/** Holds all widgets for the profiler window like menu bar, toolbar and tabs. */
+		TSharedRef<SVerticalBox> MainContentPanel = SNew(SVerticalBox);
+		ChildSlot
+			[
+				SNew(SBorder)
+				.BorderImage(FCoreStyle::Get().GetBrush("ToolPanel.GroupBorder"))
+				[
+					SNew(SOverlay)
+
+					// Overlay slot for the main HLOD window area
+					+ SOverlay::Slot()
+					[
+						MainContentPanel
+					]
+				]
+			];
+
+		MainContentPanel->AddSlot()
+			.AutoHeight()
+			.Padding(0.0f, 0.0f, 0.0f, 4.0f)
+			[
+				CreateButtonWidgets()
+			];
+
+		MainContentPanel->AddSlot()
+			.FillHeight(0.6f)
+			[
+				CreateTreeviewWidget()
+			];
+
+		MainContentPanel->AddSlot()
+			.FillHeight(0.4f)
+			[
+				SettingsView.ToSharedRef()
+			];
+
+		RegisterDelegates();
+
+		// Register to update when an undo/redo operation has been called to update our list of actors
+		GEditor->RegisterForUndo(this);	
+	}
+
+	TSharedRef<SWidget> SHLODOutliner::CreateButtonWidgets()
+	{
+		return SNew(SVerticalBox)
+			+ SVerticalBox::Slot()
+			.Padding(FMargin(0.0f, 5.0f))
+			[
+				SNew(SHorizontalBox)
+				+ SHorizontalBox::Slot()
+				.AutoWidth()
+				.VAlign(VAlign_Center)
+				.Padding(FMargin(5.0f, 0.0f))
+				[
+					SNew(SButton)
+					.HAlign(HAlign_Center)
+					.Text(LOCTEXT("BuildHLODS", "Build HLODs"))
+					.OnClicked(this, &SHLODOutliner::HandleBuildHLODs)
+				]
+
+				+ SHorizontalBox::Slot()
+					.AutoWidth()
+					.VAlign(VAlign_Center)
+					.Padding(FMargin(5.0f, 0.0f))
+					[
+						SNew(SButton)						
+						.HAlign(HAlign_Center)
+						.Text(LOCTEXT("DeleteHLODS", "Delete HLODs"))
+						.OnClicked(this, &SHLODOutliner::HandleDeleteHLODs)
+					]
+
+				+ SHorizontalBox::Slot()
+					.AutoWidth()
+					.VAlign(VAlign_Center)
+					.Padding(FMargin(5.0f, 0.0f))
+					[
+						SNew(SButton)						
+						.HAlign(HAlign_Center)
+						.Text(LOCTEXT("PreviewHLODS", "Preview HLODs"))
+						.OnClicked(this, &SHLODOutliner::HandlePreviewHLODs)
+					]
+
+				+ SHorizontalBox::Slot()
+					.AutoWidth()
+					.VAlign(VAlign_Center)
+					.Padding(FMargin(5.0f, 0.0f))
+					[
+						SNew(SButton)
+						.HAlign(HAlign_Center)
+						.Text(LOCTEXT("DeletePreviewHLODS", "Delete Preview HLODs"))
+						.OnClicked(this, &SHLODOutliner::HandleDeletePreviewHLODs)
+					]
+
+				
+			]
+			+ SVerticalBox::Slot()
+			.Padding(FMargin(0.0f, 5.0f))
+			[
+				SNew(SHorizontalBox)
+				+SHorizontalBox::Slot()
+				.AutoWidth()
+				.VAlign(VAlign_Center)
+				.Padding(FMargin(5.0f, 0.0f))
+				[
+					SNew(SButton)
+					.HAlign(HAlign_Center)
+					.Text(LOCTEXT("BuildLODActors", "Build meshes for LODActors"))
+					.OnClicked(this, &SHLODOutliner::HandleBuildLODActors)
+				]
+
+#if 0 
+
+				+ SHorizontalBox::Slot()
+				.AutoWidth()
+				.VAlign(VAlign_Center)
+				.Padding(FMargin(5.0f, 0.0f))
+				[
+					SNew(SButton)
+					.HAlign(HAlign_Center)
+					.Text(LOCTEXT("ForceRefresh", "Force refresh"))
+					.OnClicked(this, &SHLODOutliner::HandleForceRefresh)
+				]
+#endif
+			];
+		
+	}
+
+	TSharedRef<SWidget> SHLODOutliner::CreateTreeviewWidget()
+	{
+		return SAssignNew(TreeView, SHLODTree)
+			.ItemHeight(24.0f)
+			.TreeItemsSource(&HLODTreeRoot)
+			.OnGenerateRow(this, &SHLODOutliner::OnOutlinerGenerateRow)
+			.OnGetChildren(this, &SHLODOutliner::OnOutlinerGetChildren)
+			.OnSelectionChanged(this, &SHLODOutliner::OnOutlinerSelectionChanged)
+			.OnMouseButtonDoubleClick(this, &SHLODOutliner::OnOutlinerDoubleClick)
+			.OnContextMenuOpening(this, &SHLODOutliner::OnOpenContextMenu)
+			.OnExpansionChanged(this, &SHLODOutliner::OnItemExpansionChanged)
+			.HeaderRow
+			(
+				SNew(SHeaderRow)
+				+ SHeaderRow::Column("SceneActorName")
+				.DefaultLabel(LOCTEXT("SceneActorName", "Scene Actor Name"))
+				.FillWidth(0.5f)
+				+ SHeaderRow::Column("ForceCheckbox")
+				.DefaultLabel(LOCTEXT("ForceCheckbox", "Force HLOD"))
+				.DefaultTooltip(LOCTEXT("ForceCheckboxTooltip", "Force Hierarchical LOD level visibility"))
+				.FillWidth(0.25f)
+				+ SHeaderRow::Column("TriangleCount")
+				.DefaultLabel(LOCTEXT("TriangleCount", "Number of Triangles"))
+				.DefaultTooltip(LOCTEXT("TriangleCountToolTip", "Number of Triangles in a LOD Mesh"))
+				.FillWidth(0.25f)				
+				);
+	}
+	void SHLODOutliner::CreateSettingsView()
+	{
+		// Create a property view
+		FPropertyEditorModule& EditModule = FModuleManager::Get().GetModuleChecked<FPropertyEditorModule>("PropertyEditor");
+
+		FNotifyHook* NotifyHook = this;
+		FDetailsViewArgs DetailsViewArgs(
+			/*bUpdateFromSelection=*/ false,
+			/*bLockable=*/ false,
+			/*bAllowSearch=*/ false,
+			FDetailsViewArgs::HideNameArea,
+			/*bHideSelectionTip=*/ true,
+			/*InNotifyHook=*/ NotifyHook,
+			/*InSearchInitialKeyFocus=*/ false,
+			/*InViewIdentifier=*/ NAME_None);
+		DetailsViewArgs.DefaultsOnlyVisibility = FDetailsViewArgs::EEditDefaultsOnlyNodeVisibility::Automatic;
+		DetailsViewArgs.bShowOptions = false;
+
+		SettingsView = EditModule.CreateDetailView(DetailsViewArgs);
+
+		struct Local
+		{
+			/** Delegate to show all properties */
+			static bool IsPropertyVisible(const FPropertyAndParent& PropertyAndParent, bool bInShouldShowNonEditable)
+			{
+				if (PropertyAndParent.Property.GetName() == "HierarchicalLODSetup" || (PropertyAndParent.ParentProperty && PropertyAndParent.ParentProperty->GetName() == "MergeSetting"))
+				{
+					return true;
+				}
+				return false;
+			}
+
+		};
+
+		SettingsView->SetIsPropertyVisibleDelegate(FIsPropertyVisible::CreateStatic(&Local::IsPropertyVisible, true));
+		SettingsView->SetDisableCustomDetailLayouts(true);
+	}
+
+	void SHLODOutliner::Tick(const FGeometry& AllottedGeometry, const double InCurrentTime, const float InDeltaTime)
+	{
+		// Get a collection of items and folders which were formerly collapsed
+		const FParentsExpansionState ExpansionStateInfo = GetParentsExpansionState();
+
+		if (bNeedsRefresh)
+		{
+			Populate();
+		}
+
+		bool bChangeMade = false;
+
+		// Only deal with 256 at a time
+		const int32 End = FMath::Min(PendingActions.Num(), 256);
+		for (int32 Index = 0; Index < End; ++Index)
+		{
+			auto& PendingAction = PendingActions[Index];
+			switch (PendingAction.Type)
+			{
+			case FOutlinerAction::AddItem:
+				bChangeMade |= AddItemToTree(PendingAction.Item, PendingAction.ParentItem);
+				break;
+
+			case FOutlinerAction::MoveItem:
+				MoveItemInTree(PendingAction.Item, PendingAction.ParentItem);
+				bChangeMade = true;
+				break;
+
+			case FOutlinerAction::RemoveItem:
+				RemoveItemFromTree(PendingAction.Item);
+				bChangeMade = true;
+				break;
+			default:
+				check(false);
+				break;
+			}
+		}
+		PendingActions.RemoveAt(0, End);
+
+		// Restore expansion states
+		SetParentsExpansionState(ExpansionStateInfo);
+
+		if (bChangeMade)
+		{
+			TreeView->RequestTreeRefresh();
+		}
+
+	}
+
+	void SHLODOutliner::OnMouseEnter(const FGeometry& MyGeometry, const FPointerEvent& MouseEvent)
+	{
+		SCompoundWidget::OnMouseEnter(MyGeometry, MouseEvent);
+	}
+
+	void SHLODOutliner::OnMouseLeave(const FPointerEvent& MouseEvent)
+	{
+		SCompoundWidget::OnMouseLeave(MouseEvent);
+	}
+
+	FReply SHLODOutliner::OnKeyDown(const FGeometry& MyGeometry, const FKeyEvent& InKeyEvent)
+	{
+		return SCompoundWidget::OnKeyDown(MyGeometry, InKeyEvent);
+	}
+
+	FReply SHLODOutliner::OnDrop(const FGeometry& MyGeometry, const FDragDropEvent& DragDropEvent)
+	{
+		return SCompoundWidget::OnDrop(MyGeometry, DragDropEvent);
+	}
+
+	FReply SHLODOutliner::OnDragOver(const FGeometry& MyGeometry, const FDragDropEvent& DragDropEvent)
+	{
+		return SCompoundWidget::OnDragOver(MyGeometry, DragDropEvent);
+	}
+
+	void SHLODOutliner::PostUndo(bool bSuccess)
+	{
+		FullRefresh();
+	}
+
+	FReply SHLODOutliner::HandleBuildHLODs()
+	{
+		if (CurrentWorld)
+		{
+			CurrentWorld->HierarchicalLODBuilder->Build();
+		}
+
+		FullRefresh();
+		return FReply::Handled();
+	}
+
+	FReply SHLODOutliner::HandleDeleteHLODs()
+	{
+		if (CurrentWorld)
+		{
+			CurrentWorld->HierarchicalLODBuilder->ClearHLODs();
+		}
+		FullRefresh();
+		return FReply::Handled();
+	}
+
+	FReply SHLODOutliner::HandlePreviewHLODs()
+	{
+		if (CurrentWorld)
+		{
+			CurrentWorld->HierarchicalLODBuilder->PreviewBuild();
+		}
+		FullRefresh();
+		return FReply::Handled();
+	}
+
+	FReply SHLODOutliner::HandleDeletePreviewHLODs()
+	{
+		if (CurrentWorld)
+		{
+			CurrentWorld->HierarchicalLODBuilder->ClearPreviewBuild();
+		}
+		FullRefresh();
+		return FReply::Handled();
+	}
+
+	FReply SHLODOutliner::HandleBuildLODActors()
+	{
+		if (CurrentWorld)
+		{
+			CurrentWorld->HierarchicalLODBuilder->BuildMeshesForLODActors();
+		}
+
+		FullRefresh();
+		return FReply::Handled();
+	}
+
+	FReply SHLODOutliner::HandleForceRefresh()
+	{
+		FullRefresh();
+
+		return FReply::Handled();
+	}
+
+	END_SLATE_FUNCTION_BUILD_OPTIMIZATION
+
+		void SHLODOutliner::RegisterDelegates()
+	{
+		FEditorDelegates::MapChange.AddSP(this, &SHLODOutliner::OnMapChange);
+		FEditorDelegates::NewCurrentLevel.AddSP(this, &SHLODOutliner::OnNewCurrentLevel);
+		FWorldDelegates::LevelAddedToWorld.AddSP(this, &SHLODOutliner::OnLevelAdded);
+		FWorldDelegates::LevelRemovedFromWorld.AddSP(this, &SHLODOutliner::OnLevelRemoved);
+		GEngine->OnLevelActorListChanged().AddSP(this, &SHLODOutliner::FullRefresh);
+		GEngine->OnLevelActorAdded().AddSP(this, &SHLODOutliner::OnLevelActorsAdded);
+		GEngine->OnLevelActorDeleted().AddSP(this, &SHLODOutliner::OnLevelActorsRemoved);
+		GEngine->OnActorMoved().AddSP(this, &SHLODOutliner::OnActorMovedEvent);
+
+		// Register to be notified when properties are edited
+		FCoreDelegates::OnActorLabelChanged.AddRaw(this, &SHLODOutliner::OnActorLabelChanged);
+		
+		// Selection change
+		USelection::SelectionChangedEvent.AddRaw(this, &SHLODOutliner::OnLevelSelectionChanged);
+		USelection::SelectObjectEvent.AddRaw(this, &SHLODOutliner::OnLevelSelectionChanged);
+				
+		// HLOD related events
+		GEngine->OnHLODActorMoved().AddSP(this, &SHLODOutliner::OnHLODActorMovedEvent);
+		GEngine->OnHLODActorAdded().AddSP(this, &SHLODOutliner::OnHLODActorAddedEvent);
+		GEngine->OnHLODActorMarkedDirty().AddSP(this, &SHLODOutliner::OnHLODActorMarkedDirtyEvent);
+		GEngine->OnHLODDrawDistanceChanged().AddSP(this, &SHLODOutliner::OnHLODDrawDistanceChangedEvent);
+	}
+
+	void SHLODOutliner::DeregisterDelegates()
+	{
+		FEditorDelegates::MapChange.RemoveAll(this);
+		FEditorDelegates::NewCurrentLevel.RemoveAll(this);
+		FWorldDelegates::LevelAddedToWorld.RemoveAll(this);
+		FWorldDelegates::LevelRemovedFromWorld.RemoveAll(this);
+		GEngine->OnLevelActorListChanged().RemoveAll(this);
+		GEngine->OnLevelActorAdded().RemoveAll(this);
+		GEngine->OnLevelActorDeleted().RemoveAll(this);
+		GEngine->OnActorMoved().RemoveAll(this);
+
+		FCoreDelegates::OnActorLabelChanged.AddRaw(this, &SHLODOutliner::OnActorLabelChanged);
+
+		USelection::SelectionChangedEvent.RemoveAll(this);
+		USelection::SelectObjectEvent.RemoveAll(this);
+
+		GEngine->OnHLODActorMoved().RemoveAll(this);
+		GEngine->OnHLODActorAdded().RemoveAll(this);
+		GEngine->OnHLODActorMarkedDirty().RemoveAll(this);
+	}
+
+	void SHLODOutliner::ForceViewLODActor(TSharedRef<ITreeItem> Item)
+	{
+		if (CurrentWorld)
+		{
+			const FScopedTransaction Transaction(LOCTEXT("UndoAction_LODLevelForcedView", "LOD Level Forced View"));
+			FLODActorItem* ActorItem = (FLODActorItem*)(&Item.Get());
+
+			if (ActorItem->LODActor.IsValid())
+			{
+				ActorItem->LODActor->Modify();
+				ActorItem->LODActor->ToggleForceView();				
+			}
+		}
+	}
+
+	ECheckBoxState SHLODOutliner::IsHLODLevelChecked(TSharedRef<ITreeItem> Item) const
+	{
+		check(Item->GetTreeItemType() == ITreeItem::HierarchicalLODLevel);
+
+		FLODLevelItem* LevelItem = static_cast<FLODLevelItem*>(&Item.Get());
+
+		return (LevelItem->LODLevelIndex == ForcedLODLevel) ? ECheckBoxState::Checked : ECheckBoxState::Unchecked;
+	}
+
+	void SHLODOutliner::HandleCheckBoxCheckedStateChanged(ECheckBoxState NewState, TSharedRef<ITreeItem> Item)
+	{
+		check(Item->GetTreeItemType() == ITreeItem::HierarchicalLODLevel);
+
+		FLODLevelItem* LevelItem = static_cast<FLODLevelItem*>(&Item.Get());
+		if (LevelItem->LODLevelIndex == ForcedLODLevel && NewState == ECheckBoxState::Unchecked)
+		{
+			// Reset all
+			RestoreForcedLODLevel(LevelItem->LODLevelIndex);
+			ForcedLODLevel = -1;
+		}
+		else if ( LevelItem->LODLevelIndex != ForcedLODLevel && NewState == ECheckBoxState::Checked)
+		{
+			RestoreForcedLODLevel(ForcedLODLevel);
+			SetForcedLODLevel(LevelItem->LODLevelIndex);			
+		}
+	}
+
+	bool SHLODOutliner::CanHLODLevelBeForced(const uint32 LODLevel) const
+	{
+		return LODLevelBuildFlags[LODLevel];
+	}
+
+	void SHLODOutliner::RestoreForcedLODLevel(const uint32 LODLevel)
+	{
+		if (LODLevel == -1)
+		{
+			return;
+		}
+
+		if (CurrentWorld)
+		{
+			auto Level = CurrentWorld->GetCurrentLevel();
+			for (auto Actor : Level->Actors)
+			{
+				ALODActor* LODActor = Cast<ALODActor>(Actor);
+				if (LODActor)
+				{
+					if (LODActor->LODLevel == LODLevel + 1)
+					{
+						LODActor->SetForcedView(false);
+					}
+					else
+					{
+						LODActor->SetHiddenFromEditorView(false, LODLevel + 1);
+					}
+				}
+			}
+		}
+	}
+
+	void SHLODOutliner::SetForcedLODLevel(const uint32 LODLevel)
+	{
+		if (CurrentWorld)
+		{
+			auto Level = CurrentWorld->GetCurrentLevel();
+			for (auto Actor : Level->Actors)
+			{
+				ALODActor* LODActor = Cast<ALODActor>(Actor);
+				if (LODActor)
+				{
+					if (LODActor->LODLevel == LODLevel + 1)
+					{
+						LODActor->SetForcedView(true);
+					}
+					else
+					{
+						LODActor->SetHiddenFromEditorView(true, LODLevel + 1);
+					}
+				}
+			}
+		}
+		ForcedLODLevel = LODLevel;
+	}
+
+	void SHLODOutliner::BuildLODActor(TSharedRef<ITreeItem> Item)
+	{
+		if (CurrentWorld)
+		{			
+			FLODActorItem* ActorItem = (FLODActorItem*)(&Item.Get());
+
+			auto Parent = ActorItem->GetParent();
+
+			ITreeItem::TreeItemType Type = Parent->GetTreeItemType();			
+			if (Type == ITreeItem::HierarchicalLODLevel)
+			{
+				FLODLevelItem* LevelItem = (FLODLevelItem*)(Parent.Get());
+				CurrentWorld->HierarchicalLODBuilder->BuildMeshForLODActor(ActorItem->LODActor.Get(), LevelItem->LODLevelIndex);
+				FullRefresh();
+			}			
+		}
+	}
+
+	void SHLODOutliner::SelectLODActor(TSharedRef<ITreeItem> Item)
+	{
+		if (CurrentWorld)
+		{
+			FLODActorItem* ActorItem = (FLODActorItem*)(&Item.Get());
+
+			if (ActorItem->LODActor.IsValid())
+			{
+				EmptySelection();
+				StartSelection();
+				SelectActorInViewport(ActorItem->LODActor.Get(), 0);
+				EndSelection();
+			}			
+		}
+	}
+
+	void SHLODOutliner::DeleteCluster(TSharedRef<ITreeItem> Item)
+	{
+		const FScopedTransaction Transaction(LOCTEXT("UndoAction_DestroyLODActor", "Delete LOD Actor"));
+		FLODActorItem* ActorItem = (FLODActorItem*)(&Item.Get());
+		
+		ALODActor* LODActor = ActorItem->LODActor.Get();
+		ALODActor* ParentActor = HierarchicalLODUtils::GetParentLODActor(LODActor);
+
+		LODActor->Modify();
+
+		if (ParentActor)
+		{
+			ParentActor->Modify();
+		}
+
+		HierarchicalLODUtils::DeleteLODActor(LODActor);
+		CurrentWorld->DestroyActor(LODActor);
+
+		if (ParentActor && !ParentActor->HasValidSubActors())
+		{			
+			DestroyLODActor(ParentActor);
+		}
+	}
+
+	void SHLODOutliner::RemoveStaticMeshActorFromCluster(TSharedRef<ITreeItem> Item)
+	{
+		if (CurrentWorld)
+		{
+			const FScopedTransaction Transaction(LOCTEXT("UndoAction_RemoveStaticMeshActorFromCluster", "Removed Static Mesh Actor From Cluster"));
+			FStaticMeshActorItem* ActorItem = (FStaticMeshActorItem*)(&Item.Get());
+			auto Parent = ActorItem->GetParent();
+			
+			ITreeItem::TreeItemType Type = Parent->GetTreeItemType();			
+			if (Type == ITreeItem::HierarchicalLODActor)
+			{
+				FLODActorItem* ParentLODActorItem = (FLODActorItem*)(Parent.Get());
+				ParentLODActorItem->LODActor->Modify();
+				ActorItem->StaticMeshActor->Modify();
+				ParentLODActorItem->LODActor->RemoveSubActor(ActorItem->StaticMeshActor.Get());
+
+				PendingActions.Emplace(FOutlinerAction::RemoveItem, Item);
+
+				if (!ParentLODActorItem->LODActor->HasValidSubActors())
+				{
+					DestroyLODActor(ParentLODActorItem->LODActor.Get());
+					PendingActions.Emplace(FOutlinerAction::RemoveItem, Parent);
+				}
+			}
+		}
+	}
+
+	void SHLODOutliner::RemoveLODActorFromCluster(TSharedRef<ITreeItem> Item)
+	{
+		if (CurrentWorld)
+		{
+			const FScopedTransaction Transaction(LOCTEXT("UndoAction_RemoveLODActorFromCluster", "Removed LOD Actor From Cluster"));
+			FLODActorItem* ActorItem = (FLODActorItem*)(&Item.Get());
+			auto Parent = ActorItem->GetParent();
+
+			ITreeItem::TreeItemType Type = Parent->GetTreeItemType();
+			if (Type == ITreeItem::HierarchicalLODActor)
+			{
+				FLODActorItem* ParentLODActorItem = (FLODActorItem*)(Parent.Get());
+				ParentLODActorItem->LODActor->Modify();
+				ActorItem->LODActor->Modify();
+				ParentLODActorItem->LODActor->RemoveSubActor(ActorItem->LODActor.Get());
+
+				PendingActions.Emplace(FOutlinerAction::RemoveItem, Item);
+
+				if (!ParentLODActorItem->LODActor->HasValidSubActors())
+				{
+					DestroyLODActor(ParentLODActorItem->LODActor.Get());
+					PendingActions.Emplace(FOutlinerAction::RemoveItem, Parent);
+				}
+			}
+		}
+	}
+
+	void SHLODOutliner::SelectContainedActors(TSharedRef<ITreeItem> Item)
+	{
+		FLODActorItem* ActorItem = (FLODActorItem*)(&Item.Get());
+
+		ALODActor* LODActor = ActorItem->LODActor.Get();
+		SelectLODActorAndContainedActorsInViewport(LODActor, 0);
+	}
+
+	void SHLODOutliner::DestroyLODActor(ALODActor* InActor)
+	{		
+		ALODActor* ParentActor = HierarchicalLODUtils::GetParentLODActor(InActor);
+
+		HierarchicalLODUtils::DeleteLODActor(InActor);
+		CurrentWorld->DestroyActor(InActor);
+
+		if (ParentActor && !ParentActor->HasValidSubActors())
+		{
+			DestroyLODActor(ParentActor);
+		}
+	}
+
+	void SHLODOutliner::UpdateDrawDistancesForLODLevel(const uint32 LODLevelIndex)
+	{
+		if (CurrentWorld)
+		{
+			auto Level = CurrentWorld->GetCurrentLevel();
+			for (auto Actor : Level->Actors)
+			{
+				ALODActor* LODActor = Cast<ALODActor>(Actor);
+				if (LODActor)
+				{
+					if (LODActor->LODLevel == LODLevelIndex + 1)
+					{
+						LODActor->LODDrawDistance = LODLevelDrawDistances[LODLevelIndex];
+						LODActor->UpdateSubActorLODParents();
+					}
+				}
+			}
+		}
+	}
+
+	TSharedRef<ITableRow> SHLODOutliner::OnOutlinerGenerateRow(FTreeItemPtr InTreeItem, const TSharedRef<STableViewBase>& OwnerTable)
+	{
+		TSharedRef<ITableRow> Widget = SNew(SHLODWidgetItem, OwnerTable)
+			.TreeItemToVisualize(InTreeItem)
+			.Outliner(this)
+			.World(CurrentWorld);
+
+		//ItemMap.Add(InTreeItem->GetHash(), InTreeItem);
+		return Widget;
+	}
+
+	void SHLODOutliner::OnOutlinerGetChildren(FTreeItemPtr InParent, TArray<FTreeItemPtr>& OutChildren)
+	{		
+		for (auto& WeakChild : InParent->GetChildren())
+		{
+			auto Child = WeakChild.Pin();
+			// Should never have bogus entries in this list
+			check(Child.IsValid());
+			OutChildren.Add(Child);
+		}
+	}
+
+	void SHLODOutliner::OnOutlinerSelectionChanged(FTreeItemPtr TreeItem, ESelectInfo::Type SelectInfo)
+	{
+		if (SelectInfo == ESelectInfo::Direct)
+		{
+			return;
+		}
+
+		EmptySelection();
+
+		SelectedNodes = TreeView->GetSelectedItems();
+
+		if (TreeItem.IsValid())
+		{
+			
+			StartSelection();
+
+			ITreeItem::TreeItemType Type = TreeItem->GetTreeItemType();
+
+			switch (Type)
+			{
+				case ITreeItem::HierarchicalLODLevel:
+				{
+					FLODLevelItem* LevelItem = (FLODLevelItem*)(TreeItem.Get());
+					const TArray<TWeakPtr<ITreeItem>>& Children = LevelItem->GetChildren();
+					for (auto& WeakChild : Children)
+					{
+						auto Child = WeakChild.Pin();
+						check(Child.IsValid());
+						FLODActorItem* ActorItem = (FLODActorItem*)(Child.Get());
+						SelectLODActorAndContainedActorsInViewport(ActorItem->LODActor.Get(), 0);
+					}
+
+					break;
+				}
+
+				case ITreeItem::HierarchicalLODActor:
+				{
+					FLODActorItem* ActorItem = (FLODActorItem*)(TreeItem.Get());
+					SelectActorInViewport(ActorItem->LODActor.Get(), 0);
+					break;
+				}
+
+				case ITreeItem::StaticMeshActor:
+				{
+					FStaticMeshActorItem* StaticMeshActorItem = (FStaticMeshActorItem*)(TreeItem.Get());
+					SelectActorInViewport(StaticMeshActorItem->StaticMeshActor.Get(), 0);
+					break;
+				}
+
+			
+			}
+			
+
+			EndSelection();
+		}
+	}
+
+	void SHLODOutliner::OnOutlinerDoubleClick(FTreeItemPtr TreeItem)
+	{
+		
+	}
+
+	TSharedPtr<SWidget> SHLODOutliner::OnOpenContextMenu()
+	{
+		if (!CurrentWorld)
+		{
+			return nullptr;
+		}
+
+		// Build up the menu for a selection
+		const bool bCloseAfterSelection = true;
+		TSharedPtr<FExtender> Extender = MakeShareable(new FExtender);
+
+		FMenuBuilder MenuBuilder(bCloseAfterSelection, TSharedPtr<FUICommandList>(), Extender);
+
+		const auto NumSelectedItems = TreeView->GetNumItemsSelected();
+		if (NumSelectedItems == 1)
+		{
+			TreeView->GetSelectedItems()[0]->GenerateContextMenu(MenuBuilder, *this);
+
+			return MenuBuilder.MakeWidget();
+		}
+
+		return TSharedPtr<SWidget>();
+	}
+
+	void SHLODOutliner::OnItemExpansionChanged(FTreeItemPtr TreeItem, bool bIsExpanded)
+	{
+		TreeItem->bIsExpanded = bIsExpanded;
+
+		// Expand any children that are also expanded
+		for (auto WeakChild : TreeItem->GetChildren())
+		{
+			auto Child = WeakChild.Pin();
+			if (Child->bIsExpanded)
+			{
+				TreeView->SetItemExpansion(Child, true);
+			}
+		}
+	}
+
+	void SHLODOutliner::StartSelection()
+	{
+		GEditor->GetSelectedActors()->BeginBatchSelectOperation();
+	}
+
+	void SHLODOutliner::EmptySelection()
+	{
+		GEditor->SelectNone(false, true, true);
+		DestroySelectionActors();
+	}
+
+	void SHLODOutliner::DestroySelectionActors()
+	{
+		if (CurrentWorld)
+		{
+			for (AHLODSelectionActor* BoundsActor : SelectionActors)
+			{
+				if (BoundsActor)
+				{
+					CurrentWorld->DestroyActor(BoundsActor);
+				}
+			}
+		}
+		SelectionActors.Empty();
+	}
+
+	void SHLODOutliner::SelectActorInViewport(AActor* Actor, const uint32 SelectionDepth)
+	{
+		GEditor->SelectActor(Actor, true, false);
+
+		if (Actor->IsA<ALODActor>() && SelectionDepth == 0)
+		{
+			CreateBoundingSphereForActor(Actor);
+		}
+	}
+
+	void SHLODOutliner::SelectLODActorAndContainedActorsInViewport(ALODActor* LODActor, const uint32 SelectionDepth)
+	{
+		TArray<AActor*> SubActors;
+		ExtractStaticMeshActorsFromLODActor(LODActor, SubActors);
+		for (AActor* SubActor : SubActors)
+		{
+			SelectActorInViewport(SubActor, SelectionDepth + 1);
+		}
+
+		SelectActorInViewport(LODActor, SelectionDepth);
+	}
+
+	UDrawSphereComponent* SHLODOutliner::CreateBoundingSphereForActor(AActor* Actor)
+	{
+		if (CurrentWorld)
+		{
+			AHLODSelectionActor* SelectionActor = CurrentWorld->SpawnActorDeferred<AHLODSelectionActor>(AHLODSelectionActor::StaticClass(), FTransform());
+			SelectionActor->ClearFlags(RF_Public | RF_Standalone);
+			SelectionActor->SetFlags(RF_Transient);			
+
+			UDrawSphereComponent* BoundSphereSpawned = SelectionActor->GetDrawSphereComponent();
+			BoundSphereSpawned->SetRelativeLocation(FVector(0.0f, 0.0f, 0.0f));
+			BoundSphereSpawned->RegisterComponent();
+
+			FVector Origin, Extent;
+			FBox BoundingBox = Actor->GetComponentsBoundingBox(true);
+			BoundSphereSpawned->SetWorldLocation(BoundingBox.GetCenter());
+			BoundSphereSpawned->SetSphereRadius(BoundingBox.GetExtent().Size());
+			BoundSphereSpawned->ShapeColor = FColor::Red;		
+
+			SelectionActors.Add(SelectionActor);
+
+			return BoundSphereSpawned;
+		}
+
+		return nullptr;
+	}
+
+	void SHLODOutliner::ExtractStaticMeshActorsFromLODActor(ALODActor* LODActor, TArray<AActor*> &InOutActors)
+	{
+		for (auto ChildActor : LODActor->SubActors)
+		{
+			if (ChildActor)
+			{
+				TArray<AActor*> ChildActors;
+				if (ChildActor->IsA<ALODActor>())
+				{
+					ExtractStaticMeshActorsFromLODActor(Cast<ALODActor>(ChildActor), ChildActors);
+				}
+
+				ChildActors.Push(ChildActor);
+				InOutActors.Append(ChildActors);
+			}
+		}
+	}
+
+	void SHLODOutliner::EndSelection()
+	{
+		// Commit selection changes
+		GEditor->GetSelectedActors()->EndBatchSelectOperation();
+
+		// Fire selection changed event
+		GEditor->NoteSelectionChange();
+	}
+
+	void SHLODOutliner::OnLevelSelectionChanged(UObject* Obj)
+	{		
+		AActor* Actor = Cast<AActor>(Obj);
+		if (Actor)
+		{
+			auto Item = TreeItemsMap.Find(Actor);
+			if (Item)
+			{
+				SelectItemInTree(*Item);
+			}	
+			else
+			{
+				DestroySelectionActors();
+			}
+		}
+	}
+
+	void SHLODOutliner::OnLevelAdded(ULevel* InLevel, UWorld* InWorld)
+	{
+		//FullRefresh();
+	}
+
+	void SHLODOutliner::OnLevelRemoved(ULevel* InLevel, UWorld* InWorld)
+	{
+		//FullRefresh();
+	}
+
+	void SHLODOutliner::OnLevelActorsAdded(AActor* InActor)
+	{
+		if (!InActor->IsA<AHLODSelectionActor>() && !InActor->IsA<AWorldSettings>())
+			FullRefresh();
+	}
+
+	void SHLODOutliner::OnLevelActorsRemoved(AActor* InActor)
+	{
+		if (!InActor->IsA<AHLODSelectionActor>() && !InActor->IsA<AWorldSettings>())
+			FullRefresh();
+	}
+	
+	void SHLODOutliner::OnActorLabelChanged(AActor* ChangedActor)
+	{
+		if (!ChangedActor->IsA<AHLODSelectionActor>())
+			FullRefresh();
+	}
+
+	void SHLODOutliner::OnMapChange(uint32 MapFlags)
+	{
+		FullRefresh();
+	}
+
+	void SHLODOutliner::OnNewCurrentLevel()
+	{
+		FullRefresh();
+	}
+
+	void SHLODOutliner::OnHLODActorMovedEvent(const AActor* InActor, const AActor* ParentActor)
+	{
+		FTreeItemPtr* TreeItem = TreeItemsMap.Find(InActor);
+		FTreeItemPtr* ParentItem = TreeItemsMap.Find(ParentActor);
+		if (TreeItem && ParentItem)
+		{			
+			PendingActions.Emplace(FOutlinerAction::MoveItem, *TreeItem, *ParentItem);
+
+			auto CurrentParent = (*TreeItem)->GetParent(); 
+
+			if (CurrentParent.IsValid())
+			{
+				if (CurrentParent->GetTreeItemType() == ITreeItem::HierarchicalLODActor)
+				{
+					FLODActorItem* ParentLODActorItem = (FLODActorItem*)CurrentParent.Get();
+					if (!ParentLODActorItem->LODActor->HasValidSubActors())
+					{
+						DestroyLODActor(ParentLODActorItem->LODActor.Get());
+						PendingActions.Emplace(FOutlinerAction::RemoveItem, CurrentParent);
+					}
+				}
+			}
+		}
+	}
+
+	void SHLODOutliner::OnActorMovedEvent(AActor* InActor)
+	{
+		if (InActor->IsA<ALODActor>())
+		{
+			return;
+		}
+
+		ALODActor* ParentActor = HierarchicalLODUtils::GetParentLODActor(InActor);
+		if (ParentActor)
+		{
+			ParentActor->Modify();
+			ParentActor->SetIsDirty(true);
+		}		
+	}
+
+	void SHLODOutliner::OnHLODActorAddedEvent(const AActor* InActor, const AActor* ParentActor)
+	{
+		FullRefresh();
+	}
+
+	void SHLODOutliner::OnHLODActorMarkedDirtyEvent(ALODActor* InActor)
+	{		
+		if (InActor->GetStaticMeshComponent()->StaticMesh)
+		{
+			HierarchicalLODUtils::DeleteLODActorAssets(InActor);		
+		}		
+		FullRefresh();
+	}
+
+	void SHLODOutliner::OnHLODDrawDistanceChangedEvent()
+	{
+		if (CurrentWorld)
+		{
+			auto WorldSettings = CurrentWorld->GetWorldSettings();
+
+			check(WorldSettings->HierarchicalLODSetup.Num() == LODLevelDrawDistances.Num());
+
+			for (int32 LODLevelIndex = 0; LODLevelIndex < LODLevelDrawDistances.Num(); ++LODLevelIndex)
+			{
+				if (LODLevelDrawDistances[LODLevelIndex] != WorldSettings->HierarchicalLODSetup[LODLevelIndex].DrawDistance)
+				{
+					LODLevelDrawDistances[LODLevelIndex] = WorldSettings->HierarchicalLODSetup[LODLevelIndex].DrawDistance;
+					UpdateDrawDistancesForLODLevel(LODLevelIndex);
+				}
+			}
+		}
+	}
+
+	void SHLODOutliner::FullRefresh()
+	{
+		bNeedsRefresh = true;
+	}
+
+	void SHLODOutliner::Populate()
+	{
+		HLODTreeRoot.Empty();
+		TreeItemsMap.Empty();
+
+		for (const FWorldContext& Context : GEngine->GetWorldContexts())
+		{
+			if (Context.WorldType == EWorldType::PIE)
+			{
+				CurrentWorld = Context.World();
+				break;
+			}
+			else if (Context.WorldType == EWorldType::Editor)
+			{
+				CurrentWorld = Context.World();
+			}
+		}
+
+		if (CurrentWorld)
+		{
+			// Update settings view
+			SettingsView->SetObject(CurrentWorld->GetWorldSettings());
+
+			TArray<FTreeItemRef> LevelNodes;
+			AWorldSettings* WorldSettings = CurrentWorld->GetWorldSettings();
+			if (WorldSettings)
+			{
+				LODLevelBuildFlags.Empty();
+				LODLevelActors.Empty();
+				LODLevelDrawDistances.Empty();
+
+				const uint32 LODLevels = WorldSettings->HierarchicalLODSetup.Num();
+				for (uint32 LODLevelIndex = 0; LODLevelIndex < LODLevels; ++LODLevelIndex)
+				{
+					FTreeItemRef LevelItem = MakeShareable(new FLODLevelItem(LODLevelIndex));
+
+					PendingActions.Emplace(FOutlinerAction::AddItem, LevelItem);
+
+					LevelNodes.Add(LevelItem->AsShared());
+					HLODTreeRoot.Add(LevelItem->AsShared());
+					AllNodes.Add(LevelItem->AsShared());
+
+					LODLevelBuildFlags.Add(true);
+					LODLevelActors.AddDefaulted();
+					LODLevelDrawDistances.Add(WorldSettings->HierarchicalLODSetup[LODLevelIndex].DrawDistance);
+
+					TreeItemsMap.Add(LevelItem->GetID(), LevelItem);
+				}
+
+				for (AActor* Actor : CurrentWorld->GetCurrentLevel()->Actors)
+				{
+					if (Actor && Actor->IsA<ALODActor>())
+					{
+						ALODActor* LODActor = CastChecked<ALODActor>(Actor);
+
+						if (LODActor)
+						{
+							FTreeItemRef Item = MakeShareable(new FLODActorItem(LODActor));
+							AllNodes.Add(Item->AsShared());
+
+							PendingActions.Emplace(FOutlinerAction::AddItem, Item, LevelNodes[LODActor->LODLevel - 1]);
+
+							for (AActor* ChildActor : LODActor->SubActors)
+							{
+								if (ChildActor->IsA<ALODActor>())
+								{
+									FTreeItemRef ChildItem = MakeShareable(new FLODActorItem(CastChecked<ALODActor>(ChildActor)));
+									AllNodes.Add(ChildItem->AsShared());
+									Item->AddChild(ChildItem);
+								}
+								else
+								{
+									FTreeItemRef ChildItem = MakeShareable(new FStaticMeshActorItem(ChildActor));
+									AllNodes.Add(ChildItem->AsShared());
+
+									PendingActions.Emplace(FOutlinerAction::AddItem, ChildItem, Item);
+								}
+							}
+
+							LODLevelBuildFlags[LODActor->LODLevel - 1] &= !LODActor->IsDirty();
+							LODLevelActors[LODActor->LODLevel - 1].Add(LODActor);
+						}
+					}
+				}
+
+				for (uint32 LODLevelIndex = 0; LODLevelIndex < LODLevels; ++LODLevelIndex)
+				{
+					if (LODLevelActors[LODLevelIndex].Num() == 0)
+					{
+						LODLevelBuildFlags[LODLevelIndex] = false;
+					}
+				}
+			}
+
+			TreeView->RequestTreeRefresh();
+		}
+
+		bNeedsRefresh = false;
+	}
+
+	TMap<FTreeItemID, bool> SHLODOutliner::GetParentsExpansionState() const
+	{
+		FParentsExpansionState States;
+		for (const auto& Pair : TreeItemsMap)
+		{
+			if (Pair.Value->GetChildren().Num())
+			{
+				States.Add(Pair.Key, Pair.Value->bIsExpanded);
+			}
+		}
+
+		return States;
+	}
+
+	void SHLODOutliner::SetParentsExpansionState(const FParentsExpansionState& ExpansionStateInfo) const
+	{
+		for (const auto& Pair : TreeItemsMap)
+		{
+			auto& Item = Pair.Value;
+			if (Item->GetChildren().Num())
+			{
+				const bool* bIsExpanded = ExpansionStateInfo.Find(Pair.Key);
+				if (bIsExpanded)
+				{
+					TreeView->SetItemExpansion(Item, *bIsExpanded);
+				}
+				else
+				{
+					TreeView->SetItemExpansion(Item, Item->bIsExpanded);
+				}
+			}
+		}
+	}
+
+	const bool SHLODOutliner::AddItemToTree(FTreeItemPtr InItem, FTreeItemPtr InParentItem)
+	{
+		const auto ItemID = InItem->GetID();
+
+		if (TreeItemsMap.Find(ItemID))
+		{
+			return false;
+		}
+
+		TreeItemsMap.Add(ItemID, InItem);
+
+		if (InParentItem.Get())
+		{
+			InParentItem->AddChild(InItem->AsShared());
+		}		
+
+		return true;
+	}
+
+	void SHLODOutliner::MoveItemInTree(FTreeItemPtr InItem, FTreeItemPtr InParentItem)
+	{
+		auto CurrentParent = InItem->Parent;
+		if (CurrentParent.IsValid())
+		{
+			CurrentParent.Pin()->RemoveChild(InItem->AsShared());
+		}
+
+		if (InParentItem.Get())
+		{
+			InParentItem->AddChild(InItem->AsShared());
+		}
+	}
+
+	void SHLODOutliner::RemoveItemFromTree(FTreeItemPtr InItem)
+	{
+		const int32 NumRemoved = TreeItemsMap.Remove(InItem->GetID());
+
+		if (!NumRemoved)
+		{
+			return;
+		}
+
+		auto ParentItem = InItem->GetParent();
+		if (ParentItem.IsValid())
+		{
+			ParentItem->RemoveChild(InItem->AsShared());
+		}
+	}
+
+	void SHLODOutliner::SelectItemInTree(FTreeItemPtr InItem)
+	{
+		auto Parent = InItem->GetParent();
+		while (Parent.IsValid() && !Parent->bIsExpanded)
+		{
+			Parent->bIsExpanded = true;
+			Parent = InItem->GetParent();
+		}
+		TreeView->SetSelection(InItem);
+
+		TreeView->RequestTreeRefresh();
+	}
+
+	FReply SHLODOutliner::RetrieveActors()
+	{
+		bNeedsRefresh = true;
+		return FReply::Handled();
+	}
+
+};
+
+#undef LOCTEXT_NAMESPACE // LOCTEXT_NAMESPACE "HLODOutliner"
\ No newline at end of file
diff --git a/Engine/Source/Editor/HierarchicalLODOutliner/Private/HLODOutliner.h b/Engine/Source/Editor/HierarchicalLODOutliner/Private/HLODOutliner.h
new file mode 100644
--- /dev/null
+++ b/Engine/Source/Editor/HierarchicalLODOutliner/Private/HLODOutliner.h
@@ -0,0 +1,432 @@
+// Copyright 1998-2015 Epic Games, Inc. All Rights Reserved.
+
+#pragma once
+
+#include "SCompoundWidget.h"
+#include "STreeView.h"
+#include "SListView.h"
+#include "STableViewBase.h"
+#include "TreeItemID.h"
+
+class AHLODSelectionActor;
+class UWorld;
+class UDrawSphereComponent;
+class IDetailsView;
+class ALODActor;
+class AActor;
+
+namespace HLODOutliner
+{
+	struct ITreeItem;
+	typedef TSharedPtr<ITreeItem> FTreeItemPtr;
+	typedef TSharedRef<ITreeItem> FTreeItemRef;
+
+	/**
+	* Outliner action class used for making changing to the Outliner's treeview 
+	*/	
+	struct FOutlinerAction
+	{
+		enum ActionType
+		{
+			AddItem,
+			RemoveItem,
+			MoveItem
+		};
+
+		FOutlinerAction(ActionType InType, FTreeItemPtr InItem) : Type(InType), Item(InItem) {};
+		FOutlinerAction(ActionType InType, FTreeItemPtr InItem, FTreeItemPtr InParentItem) : Type(InType), Item(InItem), ParentItem( InParentItem) {};
+
+		ActionType Type;
+		FTreeItemPtr Item;
+		FTreeItemPtr ParentItem;
+	};
+
+	/**
+	* Implements the profiler window.
+	*/
+	class SHLODOutliner : public SCompoundWidget, public FNotifyHook, public FEditorUndoClient
+	{
+	typedef STreeView<FTreeItemPtr> SHLODTree;
+	friend struct FLODActorItem;
+	friend struct FHLODTreeWidgetItem;
+	friend struct FStaticMeshActorItem;
+	friend class SHLODWidgetItem;
+
+	public:
+		/** Default constructor. */
+		SHLODOutliner();
+
+		/** Virtual destructor. */
+		virtual ~SHLODOutliner();
+
+		SLATE_BEGIN_ARGS(SHLODOutliner){}
+		SLATE_END_ARGS()
+
+		/**
+		* Constructs this widget.
+		*/
+		void Construct(const FArguments& InArgs);
+		/** Creates the panel's ButtonWidget rows */
+		TSharedRef<SWidget> CreateButtonWidgets();
+
+		/** Creates the panel's Tree view widget*/
+		TSharedRef<SWidget> CreateTreeviewWidget();
+
+		/** Initializes and creates the settings view */
+		void CreateSettingsView();
+
+		// Begin SCompoundWidget Interface
+		virtual void Tick(const FGeometry& AllottedGeometry, const double InCurrentTime, const float InDeltaTime) override;
+		virtual void OnMouseEnter(const FGeometry& MyGeometry, const FPointerEvent& MouseEvent) override;
+		virtual void OnMouseLeave(const FPointerEvent& MouseEvent) override;
+		virtual FReply OnKeyDown(const FGeometry& MyGeometry, const FKeyEvent& InKeyEvent) override;
+		virtual FReply OnDrop(const FGeometry& MyGeometry, const FDragDropEvent& DragDropEvent) override;
+		virtual FReply OnDragOver(const FGeometry& MyGeometry, const FDragDropEvent& DragDropEvent)  override;
+		// End of SCompoundWidget Interface
+
+		// Begin FEditorUndoClient Interface
+		virtual void PostUndo(bool bSuccess) override;
+		virtual void PostRedo(bool bSuccess) override { PostUndo(bSuccess); }
+		// End of FEditorUndoClient	
+
+		/** Button handlers */
+		FReply HandleBuildHLODs();
+		FReply HandleDeleteHLODs();
+		FReply HandlePreviewHLODs();
+		FReply HandleDeletePreviewHLODs();
+		FReply RetrieveActors();
+		FReply HandleBuildLODActors();
+		FReply HandleForceRefresh();
+		/** End button handlers */
+
+	private:
+		/** Registers all the callback delegates required for keeping the treeview sync */
+		void RegisterDelegates();
+
+		/** De-registers all the callback delegates required for keeping the treeview sync */
+		void DeregisterDelegates();
+
+	protected:
+		/**
+		* Forces viewing the mesh of a Cluster (LODActor)
+		*
+		* @param Item - Treeview Item representation for LODActor
+		*/
+		void ForceViewLODActor(TSharedRef<ITreeItem> Item);
+
+		/**
+		* UI callback for the state of Force HLOD View checkbox
+		*
+		* @param Item - Treeview Item representation for Hierarchical LOD level
+		* @return ECheckBoxState
+		*/
+		ECheckBoxState IsHLODLevelChecked(TSharedRef<ITreeItem> Item) const;
+
+		/**
+		* UI callback for changes in Force HLOD view checkbox
+		*
+		* @param NewState - State of the check Box
+		* @param Item - Treeview Item representation for Hierarchical LOD level
+		*/
+		void HandleCheckBoxCheckedStateChanged(ECheckBoxState NewState, TSharedRef<ITreeItem> Item);
+
+		/**
+		* Returns whether or not a HLOD level can be forced, this depends on whether or not all of the Clusters with LODLevel are non-dirty (have their meshes build)
+		*
+		* @param LODLevel - LOD level to check for
+		* @return bool
+		*/
+		bool CanHLODLevelBeForced(const uint32 LODLevel) const;
+
+		/**
+		* Restores the forced viewing state for the given LOD levle
+		*
+		* @param LODLevel - LOD level to force
+		*/
+		void RestoreForcedLODLevel(const uint32 LODLevel);
+
+		/**
+		* Forces LODActors within the given LODLevel to show their meshes (other levels hide theirs)
+		*
+		* @param LODLevel - LOD level to force
+		*/
+		void SetForcedLODLevel(const uint32 LODLevel);
+
+	protected:
+		/**
+		* Builds the HLOD mesh for the given ALODActor (cluster)
+		*
+		* @param Item - (LODActor)TreeItem containing the ALODActor ptr
+		*/
+		void BuildLODActor(TSharedRef<ITreeItem> Item);	
+
+		/**
+		* Select the LODActor in the Editor Viewport
+		*
+		* @param Item - (LODActor)TreeItem containing the ALODActor ptr
+		*/
+		void SelectLODActor(TSharedRef<ITreeItem> Item);
+
+		/**
+		* Deletes a cluster (LODActor)
+		*
+		* @param Item - Treeview Item representation for LODActor
+		*/
+		void DeleteCluster(TSharedRef<ITreeItem> Item);
+
+		/**
+		* Selects the contained actors (SubActors) for a specific LODActor
+		*
+		* @param Item - Treeview Item representation for LODActor
+		*/
+		void SelectContainedActors(TSharedRef<ITreeItem> Item);
+
+		/**
+		* Removes the given StaticMeshActor from its parent's (ALODActor) sub-actors array
+		*
+		* @param Item - (StaticMeshActor)TreeItem containing the Actor ptr
+		*/
+		void RemoveStaticMeshActorFromCluster(TSharedRef<ITreeItem> Item);
+
+		/**
+		* Removes the given LODActor from its parent's (ALODActor) sub-actors array
+		*
+		* @param Item - Treeview Item representation for LODActor
+		*/
+		void RemoveLODActorFromCluster(TSharedRef<ITreeItem> Item);
+		
+		/**
+		* Destroys an LODActor instance
+		*
+		* @param InActor - ALODActor to destroy
+		*/
+		void DestroyLODActor(ALODActor* InActor);
+
+		/**
+		* Updates the DrawDistance value for all the LODActors with the given LODLevelIndex
+		*
+		* @param LODLevelIndex -
+		*/
+		void UpdateDrawDistancesForLODLevel(const uint32 LODLevelIndex);
+
+	protected:
+		/** Tree view callbacks */
+
+		/**
+		* Generates a Tree view row for the given Tree view Node
+		*
+		* @param InReflectorNode - Node to generate a row for
+		* @param OwnerTable - Owning table of the InReflectorNode
+		* @return TSharedRef<ITableRow>
+		*/
+		TSharedRef<ITableRow> OnOutlinerGenerateRow(FTreeItemPtr InReflectorNode, const TSharedRef<STableViewBase>& OwnerTable);
+	
+		/**
+		* Treeview callback for retrieving the children of a specific TreeItem
+		*
+		* @param InParent - Parent item to return the children from
+		* @param OutChildren - InOut array for the children
+		*/
+		void OnOutlinerGetChildren(FTreeItemPtr InParent, TArray<FTreeItemPtr>& OutChildren);
+
+		/**
+		* Handles the event fired when selection with the HLOD Tree view changes
+		*
+		* @param TreeItem - Selected node(s)
+		* @param SelectionInfo - Type of selection change
+		*/
+		void OnOutlinerSelectionChanged(FTreeItemPtr TreeItem, ESelectInfo::Type /*SelectionInfo*/);
+	
+		/**
+		* Handles double click events from the HLOD Tree view
+		*
+		* @param TreeItem - Node which was double-clicked
+		*/
+		void OnOutlinerDoubleClick(FTreeItemPtr TreeItem);
+
+		/** Open a context menu for this scene outliner */
+		TSharedPtr<SWidget> OnOpenContextMenu();
+
+		/**
+		* Handles item expansion events from the HLOD tree view
+		*
+		* @param TreeItem - Item which expansion state was changed
+		* @param bIsExpanded - New expansion state
+		*/
+		void OnItemExpansionChanged(FTreeItemPtr TreeItem, bool bIsExpanded);
+		
+		/** End of Tree view callbacks */
+
+	private:
+		/** Starts the Editor selection batch */
+		void StartSelection();
+
+		/** Empties the current Editor selection */
+		void EmptySelection();
+
+		/** Destroys the created selection actors */
+		void DestroySelectionActors();		
+
+		/**
+		* Selects an Actor in the Editor viewport
+		*
+		* @param Actor - AActor to select inside the viewport
+		* @param SelectionDepth - (recursive)
+		*/
+		void SelectActorInViewport(AActor* Actor, const uint32 SelectionDepth = 0);
+
+		/**
+		* Selects actors and sub-actors for the given LODActor
+		*
+		* @param LODActor - Actor to select + subactors
+		* @param SelectionDepth - (recursive)
+		*/
+		void SelectLODActorAndContainedActorsInViewport(ALODActor* LODActor, const uint32 SelectionDepth = 0);
+
+		/**
+		* Creates a ASelectionActor for the given actor, "procedurally" drawing its bounds
+		*
+		* @param Actor - Actor to create the SelectionActor for
+		* @return UDrawSphereComponent*
+		*/
+		UDrawSphereComponent* CreateBoundingSphereForActor(AActor* Actor);
+
+		/**
+		* Extracts all the Static Mesh Actors from the given LODActor's SubActors array
+		*
+		* @param LODActor - LODActors to check the SubActors array for
+		* @param InOutActors - Array to fill with Static Mesh Actors
+		*/
+		void ExtractStaticMeshActorsFromLODActor(ALODActor* LODActor, TArray<AActor*> &InOutActors);
+	
+		/** Ends the Editor selection batch */
+		void EndSelection();
+
+	protected:
+		/** Broadcast event delegates */
+
+		/** Called by USelection::SelectionChangedEvent delegate when the level's selection changes */
+		void OnLevelSelectionChanged(UObject* Obj);
+
+		/** Called by the engine when a level is added to the world. */
+		void OnLevelAdded(ULevel* InLevel, UWorld* InWorld);
+
+		/** Called by the engine when a level is removed from the world. */
+		void OnLevelRemoved(ULevel* InLevel, UWorld* InWorld);
+
+		/** Called by the engine when an actor is added to the world. */
+		void OnLevelActorsAdded(AActor* InActor);
+
+		/** Called by the engine when an actor is remove from the world. */
+		void OnLevelActorsRemoved(AActor* InActor);
+
+		/** Handler for when a property changes on any object */
+		void OnActorLabelChanged(AActor* ChangedActor);
+
+		/** Called when the map has changed*/
+		void OnMapChange(uint32 MapFlags);
+
+		/** Called when the current level has changed */
+		void OnNewCurrentLevel();
+
+		/** Called when a HLODActor is moved between clusters */
+		void OnHLODActorMovedEvent(const AActor* InActor, const AActor* ParentActor);
+
+		/** Called when an Actor is moved inside of the level */
+		void OnActorMovedEvent(AActor* InActor);
+
+		/** Called when and HLODActor is added to the level */
+		void OnHLODActorAddedEvent(const AActor* InActor, const AActor* ParentActor);
+
+		/** Called when and HLODActor is added to the level */
+		void OnHLODActorMarkedDirtyEvent(ALODActor* InActor);	
+
+		/** Called when a DrawDistance value within WorldSettings changed */
+		void OnHLODDrawDistanceChangedEvent();
+
+		/** End of Broadcast event delegates */
+
+	private:
+		/** Tells the scene outliner that it should do a full refresh, which will clear the entire tree and rebuild it from scratch. */
+		void FullRefresh();
+
+		/** Populates the HLODTreeRoot array and consequently the Treeview */
+		void Populate();
+	
+		/** Structure containing information relating to the expansion state of parent items in the tree */
+		typedef TMap<FTreeItemID, bool> FParentsExpansionState;
+
+		/** Gets the current expansion state of parent items */
+		TMap<FTreeItemID, bool> GetParentsExpansionState() const;
+
+		/** Updates the expansion state of parent items after a repopulate, according to the previous state */
+		void SetParentsExpansionState(const FParentsExpansionState& ExpansionStateInfo) const;
+
+		/**
+		* Adds a new Treeview item
+		*
+		* @param InItem - Item to add
+		* @param InParentItem - Optional parent item to add it to
+		* @return const bool
+		*/
+		const bool AddItemToTree(FTreeItemPtr InItem, FTreeItemPtr InParentItem);
+
+		/**
+		* Moves a TreeView item around
+		*
+		* @param InItem - Item to move
+		* @param InParentItem - New parent for InItem to move to
+		*/
+		void MoveItemInTree(FTreeItemPtr InItem, FTreeItemPtr InParentItem);
+
+		/**
+		* Removes a TreeView item
+		*
+		* @param InItem - Item to remove
+		*/
+		void RemoveItemFromTree(FTreeItemPtr InItem);
+
+		/**
+		* Selects a Treeview item
+		*
+		* @param InItem - Item to select
+		*/
+		void SelectItemInTree(FTreeItemPtr InItem);
+	private:		
+		/** Whether or not we need to do a refresh of the Tree view*/
+		bool bNeedsRefresh;
+
+		/** World instance we are currently representing/mirroring in the panel */
+		UWorld* CurrentWorld;
+
+		/** Tree view nodes */
+		TArray<FTreeItemPtr> HLODTreeRoot;
+		/** Currently selected Tree view nodes*/
+		TArray<FTreeItemPtr> SelectedNodes;
+		/** HLOD Treeview widget*/
+		TSharedPtr<SHLODTree> TreeView;
+		/** Property viewing widget */
+		TSharedPtr<IDetailsView> SettingsView;
+	
+		/** Map containing all the nodes with their corresponding keys */
+		TMap<FTreeItemID, FTreeItemPtr> TreeItemsMap;
+
+		/** Array of pending OutlinerActions */
+		TArray<FOutlinerAction> PendingActions;
+
+		/** Array containing all the nodes */
+		TArray<FTreeItemPtr> AllNodes;
+	
+		/** Array of SelectionActors created for current HLOD selection*/
+		TArray<AHLODSelectionActor*> SelectionActors;
+
+		/** Currently forced LOD level*/
+		int32 ForcedLODLevel;
+		/** Array with flags for each LOD level (whether or not all their Clusters/LODActors have their meshes built) */
+		TArray<bool> LODLevelBuildFlags;
+		/** Array of LODActors/Cluster per LOD level*/
+		TArray<TArray<ALODActor*>> LODLevelActors;
+		/** Array of DrawDistances for each LOD Level*/
+		TArray<float> LODLevelDrawDistances;
+	};
+};
diff --git a/Engine/Source/Editor/HierarchicalLODOutliner/Private/HLODOutlinerDragDrop.cpp b/Engine/Source/Editor/HierarchicalLODOutliner/Private/HLODOutlinerDragDrop.cpp
new file mode 100644
--- /dev/null
+++ b/Engine/Source/Editor/HierarchicalLODOutliner/Private/HLODOutlinerDragDrop.cpp
@@ -0,0 +1,159 @@
+// Copyright 1998-2015 Epic Games, Inc. All Rights Reserved.
+
+#pragma once
+
+#include "HierarchicalLODOutlinerPrivatePCH.h"
+#include "HLODOutlinerDragDrop.h"
+#include "HierarchicalLODUtils.h"
+#include "ITreeItem.h"
+
+HLODOutliner::FDragDropPayload::FDragDropPayload()
+{
+	// QQ change this
+	LODActors = TArray<TWeakObjectPtr<AActor>>();
+	StaticMeshActors = TArray<TWeakObjectPtr<AActor>>();
+	bSceneOutliner = false;
+}
+
+bool HLODOutliner::FDragDropPayload::ParseDrag(const FDragDropOperation& Operation)
+{
+	bool bApplicable = false;
+
+	if (Operation.IsOfType<FHLODOutlinerDragDropOp>())
+	{
+		const auto& OutlinerOp = static_cast<const FHLODOutlinerDragDropOp&>(Operation);
+
+		if (OutlinerOp.StaticMeshActorOp.IsValid())
+		{
+			StaticMeshActors = OutlinerOp.StaticMeshActorOp->Actors;
+		}
+
+		if (OutlinerOp.LODActorOp.IsValid())
+		{
+			LODActors = OutlinerOp.LODActorOp->Actors;
+		}
+
+		bApplicable = true;
+		bSceneOutliner = false;
+	}
+	else if (Operation.IsOfType<FActorDragDropGraphEdOp>())
+	{
+		const auto& ActorOp = static_cast<const FActorDragDropGraphEdOp&>(Operation);
+		for (auto& ActorPtr : ActorOp.Actors)
+		{
+			AActor* Actor = ActorPtr.Get();
+			if (HierarchicalLODUtils::ShouldGenerateCluster(Actor))
+			{
+				if (!StaticMeshActors.IsSet())
+				{
+					StaticMeshActors = TArray<TWeakObjectPtr<AActor>>();
+				}
+
+				StaticMeshActors->Add(Actor);
+			}
+		}
+		bSceneOutliner = true;
+		bApplicable = true;
+	}
+
+	return bApplicable;
+}
+
+TSharedPtr<FDragDropOperation> HLODOutliner::CreateDragDropOperation(const TArray<FTreeItemPtr>& InTreeItems)
+{
+	FDragDropPayload DraggedObjects;
+	for (const auto& Item : InTreeItems)
+	{
+		Item->PopulateDragDropPayload(DraggedObjects);
+	}
+
+	if (DraggedObjects.LODActors.IsSet() || DraggedObjects.StaticMeshActors.IsSet())
+	{
+		TSharedPtr<FHLODOutlinerDragDropOp> OutlinerOp = MakeShareable(new FHLODOutlinerDragDropOp(DraggedObjects));
+		OutlinerOp->Construct();
+		return OutlinerOp;
+	}
+	
+	return nullptr;
+}
+
+HLODOutliner::FHLODOutlinerDragDropOp::FHLODOutlinerDragDropOp(const FDragDropPayload& DraggedObjects)
+	: OverrideText()
+	, OverrideIcon(nullptr)
+{
+	if (DraggedObjects.StaticMeshActors)
+	{
+		StaticMeshActorOp = MakeShareable(new FActorDragDropOp);
+		StaticMeshActorOp->Init(DraggedObjects.StaticMeshActors.GetValue());
+	}
+	
+	if (DraggedObjects.LODActors)
+	{
+		LODActorOp = MakeShareable(new FActorDragDropOp);
+		LODActorOp->Init(DraggedObjects.LODActors.GetValue());
+	}
+}
+
+TSharedPtr<SWidget> HLODOutliner::FHLODOutlinerDragDropOp::GetDefaultDecorator() const
+{
+	TSharedRef<SVerticalBox> VerticalBox = SNew(SVerticalBox);
+
+	VerticalBox->AddSlot()
+		[
+			SNew(SBorder)
+			.BorderImage(FEditorStyle::GetBrush("Graph.ConnectorFeedback.Border"))
+			.Visibility(this, &FHLODOutlinerDragDropOp::GetOverrideVisibility)
+			.Content()
+			[
+				SNew(SHorizontalBox)
+
+				+ SHorizontalBox::Slot()
+				.AutoWidth()
+				.Padding(0.0f, 0.0f, 3.0f, 0.0f)
+				[
+					SNew(SImage)
+					.Image(this, &FHLODOutlinerDragDropOp::GetOverrideIcon)
+				]
+
+				+ SHorizontalBox::Slot()
+					.AutoWidth()
+					.VAlign(VAlign_Center)
+					[
+						SNew(STextBlock)
+						.Text(this, &FHLODOutlinerDragDropOp::GetOverrideText)
+					]
+			]
+		];
+
+	if (LODActorOp.IsValid())
+	{
+		auto Content = LODActorOp->GetDefaultDecorator();
+		if (Content.IsValid())
+		{
+			Content->SetVisibility(TAttribute<EVisibility>(this, &FHLODOutlinerDragDropOp::GetDefaultVisibility));
+			VerticalBox->AddSlot()[Content.ToSharedRef()];
+		}
+	}
+
+	if (StaticMeshActorOp.IsValid())
+	{
+		auto Content = StaticMeshActorOp->GetDefaultDecorator();
+		if (Content.IsValid())
+		{
+			Content->SetVisibility(TAttribute<EVisibility>(this, &FHLODOutlinerDragDropOp::GetDefaultVisibility));
+			VerticalBox->AddSlot()[Content.ToSharedRef()];
+		}
+	}
+
+	return VerticalBox;
+}
+
+EVisibility HLODOutliner::FHLODOutlinerDragDropOp::GetOverrideVisibility() const
+{
+	return OverrideText.IsEmpty() && OverrideIcon == nullptr ? EVisibility::Collapsed : EVisibility::Visible;
+}
+
+EVisibility HLODOutliner::FHLODOutlinerDragDropOp::GetDefaultVisibility() const
+{
+	return OverrideText.IsEmpty() && OverrideIcon == nullptr ? EVisibility::Visible : EVisibility::Collapsed;
+}
diff --git a/Engine/Source/Editor/HierarchicalLODOutliner/Private/HLODOutlinerDragDrop.h b/Engine/Source/Editor/HierarchicalLODOutliner/Private/HLODOutlinerDragDrop.h
new file mode 100644
--- /dev/null
+++ b/Engine/Source/Editor/HierarchicalLODOutliner/Private/HLODOutlinerDragDrop.h
@@ -0,0 +1,149 @@
+// Copyright 1998-2015 Epic Games, Inc. All Rights Reserved.
+
+#pragma once
+#include "HierarchicalLODOutlinerPrivatePCH.h"
+#include "ActorDragDropGraphEdOp.h"
+
+class ALODActor;
+
+namespace HLODOutliner
+{
+	struct ITreeItem;
+	typedef TSharedPtr<ITreeItem> FTreeItemPtr;
+
+	/** Consilidated drag/drop information parsed for the HLOD outliner */
+	struct FDragDropPayload
+	{
+		/** Default constructor, resulting in unset contents */
+		FDragDropPayload();
+
+		/** Populate this payload from an array of tree items */
+		template<typename TItem>
+		FDragDropPayload(const TArray<TItem>& InDraggedItems)
+		{
+			for (const auto& Item : InDraggedItems)
+			{
+				Item->PopulateDragDropPayload(*this);
+			}
+		}
+
+		/** Optional array of dragged LOD actors */
+		TOptional<TArray<TWeakObjectPtr<AActor>>> LODActors;
+
+		/** OPtional array of dragged static mesh actors */
+		TOptional<TArray<TWeakObjectPtr<AActor>>> StaticMeshActors;
+
+		bool bSceneOutliner;
+
+		/**
+		*	Parse a drag operation into our list of actors and folders
+		*	@return true if the operation is viable for the HLOD outliner to process, false otherwise
+		*/
+		bool ParseDrag(const FDragDropOperation& Operation);
+	};
+
+	/** Construct a new Drag and drop operation for a scene outliner */
+	TSharedPtr<FDragDropOperation> CreateDragDropOperation(const TArray<FTreeItemPtr>& InTreeItems);
+
+	/** A drag/drop operation that was started from the scene outliner */
+	struct FHLODOutlinerDragDropOp : public FDragDropOperation
+	{
+		enum ToolTipTextType
+		{
+			ToolTip_Compatible,
+			ToolTip_Incompatible,
+			ToolTip_MultipleSelection_Compatible,
+			ToolTip_MultipleSelection_Incompatible,
+			ToolTip_CompatibleMoveToCluster,
+			ToolTip_MultipleSelection_CompatibleMoveToCluster,
+			ToolTip_CompatibleAddToCluster,
+			ToolTip_MultipleSelection_CompatibleAddToCluster,
+			ToolTip_CompatibleMergeCluster,
+			ToolTip_MultipleSelection_CompatibleMergeClusters,
+			ToolTip_CompatibleChildCluster,
+			ToolTip_MultipleSelection_CompatibleChildClusters,
+			ToolTip_CompatibleNewCluster,
+			ToolTip_MultipleSelection_CompatibleNewCluster
+		};
+
+		DRAG_DROP_OPERATOR_TYPE(FHLODOutlinerDragDropOp, FDragDropOperation);
+
+		FHLODOutlinerDragDropOp(const FDragDropPayload& DraggedObjects);
+
+		using FDragDropOperation::Construct;
+
+		/** Actor drag operation */
+		TSharedPtr<FActorDragDropOp>	StaticMeshActorOp;
+		TSharedPtr<FActorDragDropOp>	LODActorOp;
+
+		void ResetTooltip()
+		{
+			OverrideText = FText();
+			OverrideIcon = nullptr;
+		}
+
+		void SetTooltip(FText InOverrideText, const FSlateBrush* InOverrideIcon)
+		{
+			OverrideText = InOverrideText;
+			OverrideIcon = InOverrideIcon;
+		}
+
+		virtual TSharedPtr<SWidget> GetDefaultDecorator() const override;
+
+	private:
+
+		EVisibility GetOverrideVisibility() const;
+		EVisibility GetDefaultVisibility() const;
+
+		FText OverrideText;
+		FText GetOverrideText() const { return OverrideText; }
+
+		const FSlateBrush* OverrideIcon;
+		const FSlateBrush* GetOverrideIcon() const { return OverrideIcon; }
+	};
+
+	/** Struct used for validation of a drag/drop operation in the scene outliner */
+	struct FDragValidationInfo
+	{
+		/** The tooltip type to display on the operation */
+		FHLODOutlinerDragDropOp::ToolTipTextType TooltipType;
+
+		/** The tooltip text to display on the operation */
+		FText ValidationText;
+
+		/** Construct this validation information out of a tooltip type and some text */
+		FDragValidationInfo(const FHLODOutlinerDragDropOp::ToolTipTextType InTooltipType, const FText InValidationText)
+			: TooltipType(InTooltipType)
+			, ValidationText(InValidationText)
+		{}
+
+		/** Return a generic invalid result */
+		static FDragValidationInfo Invalid()
+		{
+			return FDragValidationInfo(FHLODOutlinerDragDropOp::ToolTip_Incompatible, FText());
+		}
+
+		/** @return true if this operation is valid, false otherwise */
+		bool IsValid() const
+		{
+			switch (TooltipType)
+			{
+			case FHLODOutlinerDragDropOp::ToolTip_Compatible:
+			case FHLODOutlinerDragDropOp::ToolTip_MultipleSelection_Compatible:
+			case FHLODOutlinerDragDropOp::ToolTip_CompatibleMoveToCluster:
+			case FHLODOutlinerDragDropOp::ToolTip_MultipleSelection_CompatibleMoveToCluster:
+			case FHLODOutlinerDragDropOp::ToolTip_CompatibleAddToCluster:
+			case FHLODOutlinerDragDropOp::ToolTip_MultipleSelection_CompatibleAddToCluster:
+			case FHLODOutlinerDragDropOp::ToolTip_CompatibleMergeCluster:
+			case FHLODOutlinerDragDropOp::ToolTip_MultipleSelection_CompatibleMergeClusters:
+			case FHLODOutlinerDragDropOp::ToolTip_CompatibleChildCluster:
+			case FHLODOutlinerDragDropOp::ToolTip_MultipleSelection_CompatibleChildClusters:
+			case FHLODOutlinerDragDropOp::ToolTip_CompatibleNewCluster:
+			case FHLODOutlinerDragDropOp::ToolTip_MultipleSelection_CompatibleNewCluster:
+				return true;
+			default:
+				return false;
+			}
+		}
+	};
+};
\ No newline at end of file
diff --git a/Engine/Source/Editor/HierarchicalLODOutliner/Private/HLODSelectionActor.cpp b/Engine/Source/Editor/HierarchicalLODOutliner/Private/HLODSelectionActor.cpp
new file mode 100644
--- /dev/null
+++ b/Engine/Source/Editor/HierarchicalLODOutliner/Private/HLODSelectionActor.cpp
@@ -0,0 +1,20 @@
+// Copyright 1998-2015 Epic Games, Inc. All Rights Reserved.
+
+#include "HierarchicalLODOutlinerPrivatePCH.h"
+#include "HLODSelectionActor.h"
+
+AHLODSelectionActor::AHLODSelectionActor(const FObjectInitializer& ObjectInitializer)
+	: Super(ObjectInitializer)
+{
+	bCanBeDamaged = false;
+	bListedInSceneOutliner = false;
+
+#if WITH_EDITORONLY_DATA	
+	DrawSphereComponent = CreateEditorOnlyDefaultSubobject<UDrawSphereComponent>(TEXT("VisualizeComponent0"));
+	if (DrawSphereComponent)
+	{
+		DrawSphereComponent->SetSphereRadius(0.0f);
+		RootComponent = DrawSphereComponent;
+	}
+#endif // WITH_EDITORONLY_DATA
+}
diff --git a/Engine/Source/Editor/HierarchicalLODOutliner/Private/HLODSelectionActor.h b/Engine/Source/Editor/HierarchicalLODOutliner/Private/HLODSelectionActor.h
new file mode 100644
--- /dev/null
+++ b/Engine/Source/Editor/HierarchicalLODOutliner/Private/HLODSelectionActor.h
@@ -0,0 +1,22 @@
+// Copyright 1998-2015 Epic Games, Inc. All Rights Reserved.
+
+#pragma once
+#include "HLODSelectionActor.generated.h"
+
+UCLASS(notplaceable, hidecategories = (Object, Collision, Display, Input, Blueprint, Transform))
+class AHLODSelectionActor : public AActor
+{
+	GENERATED_UCLASS_BODY()
+
+protected:
+#if WITH_EDITORONLY_DATA
+	/** Visualization component for rendering the LODActors bounds (cluster bounds)*/
+	UPROPERTY()
+	class UDrawSphereComponent* DrawSphereComponent;
+#endif // WITH_EDITORONLY_DATA
+public:
+#if WITH_EDITORONLY_DATA
+	/** Returns StaticMeshComponent subobject **/
+	class UDrawSphereComponent* GetDrawSphereComponent() const { return DrawSphereComponent; }
+#endif // WITH_EDITORONLY_DATA
+};
diff --git a/Engine/Source/Editor/HierarchicalLODOutliner/Private/HLODTreeWidgetItem.cpp b/Engine/Source/Editor/HierarchicalLODOutliner/Private/HLODTreeWidgetItem.cpp
new file mode 100644
--- /dev/null
+++ b/Engine/Source/Editor/HierarchicalLODOutliner/Private/HLODTreeWidgetItem.cpp
@@ -0,0 +1,218 @@
+// Copyright 1998-2015 Epic Games, Inc. All Rights Reserved.
+
+#include "HierarchicalLODOutlinerPrivatePCH.h"
+#include "HLODTreeWidgetItem.h"
+#include "ITreeItem.h"
+#include "HLODOutlinerDragDrop.h"
+#include "LODLevelItem.h"
+#include "LODActorItem.h"
+#include "StaticMeshActorItem.h"
+#include "HLODOutliner.h"
+
+namespace HLODOutliner
+{
+	static void UpdateOperationDecorator(const FDragDropEvent& Event, const FDragValidationInfo& ValidationInfo)
+	{
+		const FSlateBrush* Icon = ValidationInfo.IsValid() ? FEditorStyle::GetBrush(TEXT("Graph.ConnectorFeedback.OK")) : FEditorStyle::GetBrush(TEXT("Graph.ConnectorFeedback.Error"));
+
+		FDragDropOperation* Operation = Event.GetOperation().Get();
+		if (Operation)
+		{
+			if (Operation->IsOfType<FHLODOutlinerDragDropOp>())
+			{
+				auto* OutlinerOp = static_cast<FHLODOutlinerDragDropOp*>(Operation);
+				OutlinerOp->SetTooltip(ValidationInfo.ValidationText, Icon);
+			}
+			else if (Operation->IsOfType<FActorDragDropGraphEdOp>())
+			{
+				auto* ActorOp = static_cast<FActorDragDropGraphEdOp*>(Operation);
+
+				if (ValidationInfo.IsValid())
+				{
+					ActorOp->SetToolTip(FActorDragDropGraphEdOp::ToolTip_CompatibleGeneric, ValidationInfo.ValidationText );
+				}
+				else
+				{ 
+					ActorOp->SetToolTip(FActorDragDropGraphEdOp::ToolTip_IncompatibleGeneric, ValidationInfo.ValidationText);
+				}				
+			}
+		}
+	}
+
+	static void ResetOperationDecorator(const FDragDropEvent& Event)
+	{
+		FDragDropOperation* Operation = Event.GetOperation().Get();
+		if (Operation)
+		{
+			if (Operation->IsOfType<FActorDragDropGraphEdOp>())
+			{
+				static_cast<FActorDragDropGraphEdOp*>(Operation)->ResetToDefaultToolTip();
+			}
+		}
+	}
+
+	static FReply OnDragDetected(const FGeometry& MyGeometry, const FPointerEvent& MouseEvent, TWeakPtr<STableViewBase> Table)
+	{
+		auto TablePtr = Table.Pin();
+		if (TablePtr.IsValid() && MouseEvent.IsMouseButtonDown(EKeys::LeftMouseButton))
+		{
+			auto TreeView = (STreeView< FTreeItemPtr >*)TablePtr.Get();
+
+			auto Operation = HLODOutliner::CreateDragDropOperation(TreeView->GetSelectedItems());
+			if (Operation.IsValid())
+			{
+				return FReply::Handled().BeginDragDrop(Operation.ToSharedRef());
+			}
+		}
+
+		return FReply::Unhandled();
+	}
+
+	FReply HandleDrop(TWeakPtr<STableViewBase> Widget, const FDragDropEvent& DragDropEvent, IDropTarget& DropTarget, FDragValidationInfo& ValidationInfo, bool bApplyDrop = false)
+	{
+		FDragDropPayload DraggedObjects;
+		// Validate now to make sure we don't doing anything we shouldn't
+		if (!DraggedObjects.ParseDrag(*DragDropEvent.GetOperation()))
+		{
+			return FReply::Unhandled();
+		}
+
+		ValidationInfo = DropTarget.ValidateDrop(DraggedObjects);
+
+		if (!ValidationInfo.IsValid())
+		{
+			// Return handled here to stop anything else trying to handle it - the operation is invalid as far as we're concerned
+			return FReply::Handled();
+		}
+
+		if (bApplyDrop)
+		{
+			DropTarget.OnDrop(DraggedObjects, ValidationInfo, Widget.Pin().ToSharedRef());
+		}
+
+		return FReply::Handled();
+	}
+
+BEGIN_SLATE_FUNCTION_BUILD_OPTIMIZATION
+	void SHLODWidgetItem::Construct(const FArguments& InArgs, const TSharedRef<STableViewBase>& InOwner)
+	{
+		this->TreeItem = InArgs._TreeItemToVisualize.Get();
+		this->Outliner = InArgs._Outliner;
+		this->World = InArgs._World;
+		check(TreeItem);
+
+		WeakTableViewBase = InOwner;
+
+		auto Args = FSuperRowType::FArguments()
+			.Padding(1)
+			.OnDragDetected_Static(HLODOutliner::OnDragDetected, TWeakPtr<STableViewBase>(InOwner));
+
+		SMultiColumnTableRow< FTreeItemPtr >::Construct(Args, InOwner);
+	}
+
+
+		TSharedRef<SWidget> SHLODWidgetItem::GenerateWidgetForColumn(const FName& ColumnName)
+	{
+		if (ColumnName == TEXT("SceneActorName"))
+		{
+			return SNew(SHorizontalBox)
+
+				+ SHorizontalBox::Slot()
+				.AutoWidth()
+				.VAlign(VAlign_Center)
+				[
+					SNew(SExpanderArrow, SharedThis(this))
+				]
+
+			+ SHorizontalBox::Slot()
+				.AutoWidth()
+				.Padding(2.0f, 0.0f)
+				.VAlign(VAlign_Center)
+				[
+					SNew(STextBlock)
+					.Text(this, &SHLODWidgetItem::GetItemDisplayString)
+					.ColorAndOpacity(this, &SHLODWidgetItem::GetTint)
+				];
+		}
+		else if (ColumnName == TEXT("ForceCheckbox") && TreeItem->GetTreeItemType() == ITreeItem::HierarchicalLODLevel)
+		{
+			FLODLevelItem* Item = static_cast<FLODLevelItem*>(TreeItem);
+
+			return SNew(SHorizontalBox)
+				+ SHorizontalBox::Slot()
+				.AutoWidth()
+				.VAlign(VAlign_Center)
+				[
+					SNew(SBox)
+					.IsEnabled(Outliner, &SHLODOutliner::CanHLODLevelBeForced, Item->LODLevelIndex)
+					[
+						SNew(SCheckBox)
+						.IsChecked(Outliner, &SHLODOutliner::IsHLODLevelChecked, TreeItem->AsShared())
+						.OnCheckStateChanged(Outliner, &SHLODOutliner::HandleCheckBoxCheckedStateChanged, TreeItem->AsShared())
+					]					
+				];			
+		}
+		else if (ColumnName == TEXT("TriangleCount") && TreeItem->GetTreeItemType() == ITreeItem::HierarchicalLODActor)
+		{
+			FLODActorItem* Item = static_cast<FLODActorItem*>(TreeItem);
+
+			return SNew(SHorizontalBox)
+				+ SHorizontalBox::Slot()
+				.AutoWidth()
+				.VAlign(VAlign_Center)
+				[
+					SNew(STextBlock)
+					.Text(Item, &FLODActorItem::GetNumTrianglesAsText)
+					.ColorAndOpacity(this, &SHLODWidgetItem::GetTint)
+				];
+		}
+		else
+		{
+			return SNullWidget::NullWidget;
+		}
+	}
+
+END_SLATE_FUNCTION_BUILD_OPTIMIZATION
+
+	FText SHLODWidgetItem::GetItemDisplayString() const
+	{
+		return FText::FromString(TreeItem->GetDisplayString());
+	}
+
+	FSlateColor SHLODWidgetItem::GetTint() const
+	{
+		return TreeItem->GetTint();
+	}
+
+	void SHLODWidgetItem::OnDragEnter(FGeometry const& MyGeometry, FDragDropEvent const& DragDropEvent)
+	{
+		if (TreeItem)
+		{
+			FDragValidationInfo ValidationInfo = FDragValidationInfo::Invalid();
+
+			HandleDrop(WeakTableViewBase, DragDropEvent, *TreeItem, ValidationInfo, false);
+			UpdateOperationDecorator(DragDropEvent, ValidationInfo);
+		}
+	}
+
+	void SHLODWidgetItem::OnDragLeave(FDragDropEvent const& DragDropEvent)
+	{
+		ResetOperationDecorator(DragDropEvent);
+	}
+
+	FReply SHLODWidgetItem::OnDragOver(const FGeometry& MyGeometry, const FDragDropEvent& DragDropEvent)
+	{
+		return FReply::Handled();
+	}
+
+	FReply SHLODWidgetItem::OnDrop(const FGeometry& MyGeometry, const FDragDropEvent& DragDropEvent)
+	{
+		if (TreeItem)
+		{
+			FDragValidationInfo ValidationInfo = FDragValidationInfo::Invalid();
+			return HandleDrop(WeakTableViewBase, DragDropEvent, *TreeItem, ValidationInfo, true);
+		}
+
+		return FReply::Unhandled();
+	}	
+};
diff --git a/Engine/Source/Editor/HierarchicalLODOutliner/Private/HLODTreeWidgetItem.h b/Engine/Source/Editor/HierarchicalLODOutliner/Private/HLODTreeWidgetItem.h
new file mode 100644
--- /dev/null
+++ b/Engine/Source/Editor/HierarchicalLODOutliner/Private/HLODTreeWidgetItem.h
@@ -0,0 +1,74 @@
+// Copyright 1998-2015 Epic Games, Inc. All Rights Reserved.
+
+#pragma once
+
+struct ITreeItem;
+
+namespace HLODOutliner
+{
+	struct ITreeItem;
+	typedef TSharedPtr<ITreeItem> FTreeItemPtr;
+	class SHLODOutliner;
+
+	/**
+	* Widget that visualizes the contents of a FReflectorNode.
+	*/
+	class SHLODWidgetItem
+		: public SMultiColumnTableRow < FTreeItemPtr >
+	{
+	public:
+
+		SLATE_BEGIN_ARGS(SHLODWidgetItem)
+			: _TreeItemToVisualize(),
+			_Outliner(),
+			_World()
+		{ }
+
+		SLATE_ARGUMENT(FTreeItemPtr, TreeItemToVisualize)
+		SLATE_ARGUMENT(SHLODOutliner*, Outliner)
+		SLATE_ARGUMENT(UWorld*, World)
+		SLATE_END_ARGS()
+
+	public:
+		/**
+		* Construct child widgets that comprise this widget.
+		*
+		* @param InArgs Declaration from which to construct this widget.
+		*/
+		void Construct(const FArguments& InArgs, const TSharedRef<STableViewBase>& InOwner);
+
+	public:
+		// SMultiColumnTableRow overrides
+		virtual TSharedRef<SWidget> GenerateWidgetForColumn(const FName& ColumnName) override;
+
+	protected:
+		/** Returns the display string of the node */
+		FText GetItemDisplayString() const;
+
+		/** Returns the tint of the node */
+		FSlateColor GetTint() const;
+		
+		// Begin STableRow interface.
+		virtual void OnDragEnter(FGeometry const& MyGeometry, FDragDropEvent const& DragDropEvent) override;
+		virtual void OnDragLeave(FDragDropEvent const& DragDropEvent) override;
+		virtual FReply OnDragOver(const FGeometry& MyGeometry, const FDragDropEvent& DragDropEvent) override;
+		virtual FReply OnDrop(const FGeometry& MyGeometry, const FDragDropEvent& DragDropEvent) override;
+		// End STableRow interface.
+
+	private:
+		/** The info about the widget that we are visualizing. */
+		ITreeItem* TreeItem;
+
+		/** Cached Display String as FText */
+		FText CachedItemName;
+
+		/** WeakPtr to the Treeview widget */
+		TWeakPtr< STableViewBase > WeakTableViewBase;
+
+		/** Pointer to the owning SHLODOutliner */
+		SHLODOutliner* Outliner;
+		
+		/** Pointer to the currently represented world (unused atm) */
+		UWorld* World;
+	};
+};
\ No newline at end of file
diff --git a/Engine/Source/Editor/HierarchicalLODOutliner/Private/HierarchicalLODOutlinerModule.cpp b/Engine/Source/Editor/HierarchicalLODOutliner/Private/HierarchicalLODOutlinerModule.cpp
new file mode 100644
--- /dev/null
+++ b/Engine/Source/Editor/HierarchicalLODOutliner/Private/HierarchicalLODOutlinerModule.cpp
@@ -0,0 +1,23 @@
+// Copyright 1998-2015 Epic Games, Inc. All Rights Reserved.
+
+#include "HierarchicalLODOutlinerPrivatePCH.h"
+#include "HierarchicalLODOutlinerModule.h"
+#include "HLODOutliner.h"
+
+void FHierarchicalLODOutlinerModule::StartupModule()
+{
+
+}
+
+void FHierarchicalLODOutlinerModule::ShutdownModule()
+{
+
+}
+
+TSharedRef<SWidget> FHierarchicalLODOutlinerModule::CreateHLODOutlinerWidget()
+{
+	TSharedRef<HLODOutliner::SHLODOutliner> HLODWindow = SNew(HLODOutliner::SHLODOutliner);
+	return HLODWindow;
+}
+
+IMPLEMENT_MODULE(FHierarchicalLODOutlinerModule, HLODPluginModule);
\ No newline at end of file
diff --git a/Engine/Source/Editor/HierarchicalLODOutliner/Private/HierarchicalLODOutlinerPrivatePCH.h b/Engine/Source/Editor/HierarchicalLODOutliner/Private/HierarchicalLODOutlinerPrivatePCH.h
new file mode 100644
--- /dev/null
+++ b/Engine/Source/Editor/HierarchicalLODOutliner/Private/HierarchicalLODOutlinerPrivatePCH.h
@@ -0,0 +1,5 @@
+// Copyright 1998-2015 Epic Games, Inc. All Rights Reserved.
+
+#pragma once
+
+#include "Core.h"
diff --git a/Engine/Source/Editor/HierarchicalLODOutliner/Private/ITreeItem.h b/Engine/Source/Editor/HierarchicalLODOutliner/Private/ITreeItem.h
new file mode 100644
--- /dev/null
+++ b/Engine/Source/Editor/HierarchicalLODOutliner/Private/ITreeItem.h
@@ -0,0 +1,116 @@
+// Copyright 1998-2015 Epic Games, Inc. All Rights Reserved.
+#pragma once
+
+#include "Engine/LodActor.h"
+#include "HLODOutlinerDragDrop.h"
+#include "TreeItemID.h"
+
+class FMenuBuilder;
+class UWorld;
+
+namespace HLODOutliner
+{
+	class SHLODOutliner;
+
+	/** Interface used for validating movement (ie, drag/dropping) operations */
+	struct IDropTarget
+	{
+		/** Called to test whether the specified payload can be dropped onto this tree item */
+		virtual FDragValidationInfo ValidateDrop(FDragDropPayload& DraggedObjects) const = 0;
+
+		/** Called to drop the specified objects on this item. Only called if ValidateDrop() allows. */
+		virtual void OnDrop(FDragDropPayload& DraggedObjects, const FDragValidationInfo& ValidationInfo, TSharedRef<SWidget> DroppedOnWidget) = 0;
+	};
+	
+	typedef TSharedRef<ITreeItem> FTreeItemRef;
+
+	struct ITreeItem : IDropTarget, TSharedFromThis < ITreeItem >
+	{
+		enum TreeItemType
+		{
+			Invalid,
+			HierarchicalLODLevel,
+			HierarchicalLODActor,
+			StaticMeshActor
+		};
+
+		friend SHLODOutliner;
+
+	protected:
+		/** Default constructor */
+		ITreeItem() : Parent(nullptr), Type(Invalid), bIsExpanded(false) {}
+		virtual ~ITreeItem() {}
+
+		/** This item's parent, if any. */
+		mutable TWeakPtr<ITreeItem> Parent;
+
+		/** Array of children contained underneath this item */
+		mutable TArray<TWeakPtr<ITreeItem>> Children;
+
+		/** Item type enum */
+		mutable TreeItemType Type;
+
+	public:
+		/** Get this item's parent. Can be nullptr. */
+		FTreeItemPtr GetParent() const
+		{
+			return Parent.Pin();
+		}
+
+		/** Add a child to this item */
+		void AddChild(FTreeItemRef Child)
+		{
+			Child->Parent = AsShared();
+			Children.Add(MoveTemp(Child));
+		}
+
+		/** Remove a child from this item */
+		void RemoveChild(const FTreeItemRef& Child)
+		{
+			if (Children.Remove(Child))
+			{
+				Child->Parent = nullptr;
+			}
+		}
+
+		/** Returns the TreeItem's type */
+		const TreeItemType GetTreeItemType()
+		{
+			return Type;
+		}
+
+		/** Get this item's children, if any. Although we store as weak pointers, they are guaranteed to be valid. */
+		FORCEINLINE const TArray<TWeakPtr<ITreeItem>>& GetChildren() const
+		{
+			return Children;
+		}
+
+		/** Flag whether or not this item is expanded in the treeview */
+		bool bIsExpanded;
+
+	public:		
+		/** Get the raw string to display for this tree item - used for sorting */
+		virtual FString GetDisplayString() const = 0;
+
+		/** Check whether it should be possible to interact with this tree item */
+		virtual bool CanInteract() const = 0;
+
+		/** Called when this item is expanded or collapsed */
+		virtual void OnExpansionChanged() {};
+
+		/** Generate a context menu for this item. Only called if *only* this item is selected. */
+		virtual void GenerateContextMenu(FMenuBuilder& MenuBuilder, SHLODOutliner& Outliner) {};
+
+		/** Populate the specified drag/drop payload with any relevant information for this type */
+		virtual void PopulateDragDropPayload(FDragDropPayload& Payload) const = 0;
+		
+		/** Returns this TreeItem's ID */
+		virtual FTreeItemID GetID() = 0;
+
+		/** Returns the Tint used for displaying this item in the Treeview*/
+		virtual FSlateColor GetTint() const
+		{
+			return FLinearColor(1.0f, 1.0f, 1.0f);
+		}
+	};
+};
diff --git a/Engine/Source/Editor/HierarchicalLODOutliner/Private/LODActorItem.cpp b/Engine/Source/Editor/HierarchicalLODOutliner/Private/LODActorItem.cpp
new file mode 100644
--- /dev/null
+++ b/Engine/Source/Editor/HierarchicalLODOutliner/Private/LODActorItem.cpp
@@ -0,0 +1,322 @@
+// Copyright 1998-2015 Epic Games, Inc. All Rights Reserved.
+
+#include "HierarchicalLODOutlinerPrivatePCH.h"
+#include "LevelEditor.h"
+#include "LODActorItem.h"
+#include "Engine/LodActor.h"
+#include "HLODOutliner.h"
+#include "HierarchicalLODUtils.h"
+#include "ScopedTransaction.h"
+#include "SlateBasics.h"
+
+
+#define LOCTEXT_NAMESPACE "LODActorItem"
+
+HLODOutliner::FLODActorItem::FLODActorItem(ALODActor* InLODActor)
+	: LODActor(InLODActor), ID(InLODActor)
+{
+	Type = ITreeItem::HierarchicalLODActor;
+}
+
+bool HLODOutliner::FLODActorItem::CanInteract() const
+{
+	return true;
+}
+
+void HLODOutliner::FLODActorItem::GenerateContextMenu(FMenuBuilder& MenuBuilder, SHLODOutliner& Outliner)
+{
+	auto SharedOutliner = StaticCastSharedRef<SHLODOutliner>(Outliner.AsShared());
+
+	MenuBuilder.AddMenuEntry(LOCTEXT("SelectLODActor", "Select LOD Actor"), FText(), FSlateIcon(), FUIAction(FExecuteAction::CreateSP(&Outliner, &SHLODOutliner::SelectLODActor, AsShared())));
+
+	MenuBuilder.AddMenuEntry(LOCTEXT("SelectContainedActors", "Select Contained Actors"), FText(), FSlateIcon(), FUIAction(FExecuteAction::CreateSP(&Outliner, &SHLODOutliner::SelectContainedActors, AsShared())));
+
+	if (LODActor->IsDirty())
+	{
+		MenuBuilder.AddMenuEntry(LOCTEXT("BuildLODActorMesh", "Build LOD Mesh"), FText(), FSlateIcon(), FUIAction(FExecuteAction::CreateSP(&Outliner, &SHLODOutliner::BuildLODActor, AsShared())));
+	}
+	else
+	{		
+		MenuBuilder.AddMenuEntry(LOCTEXT("ForceView", "ForceView"), FText(), FSlateIcon(), FUIAction(FExecuteAction::CreateSP(&Outliner, &SHLODOutliner::ForceViewLODActor, AsShared())));
+	}
+
+	AActor* Actor = LODActor.Get();
+	ALODActor* ParentActor = HierarchicalLODUtils::GetParentLODActor(Actor);
+	if (ParentActor && Parent.Pin()->GetTreeItemType() == TreeItemType::HierarchicalLODActor)
+	{		
+		MenuBuilder.AddMenuEntry(LOCTEXT("RemoveChildFromCluster", "Remove from cluster"), FText(), FSlateIcon(), FUIAction(FExecuteAction::CreateSP(&Outliner, &SHLODOutliner::RemoveLODActorFromCluster, AsShared())));
+	}
+	else
+	{
+		MenuBuilder.AddMenuEntry(LOCTEXT("DeleteCluster", "Delete Cluster"), FText(), FSlateIcon(), FUIAction(FExecuteAction::CreateSP(&Outliner, &SHLODOutliner::DeleteCluster, AsShared())));
+	}
+}
+
+FString HLODOutliner::FLODActorItem::GetDisplayString() const
+{
+	if (ALODActor* ActorPtr = LODActor.Get())
+	{
+		return ActorPtr->GetName() + ((ActorPtr->IsDirty()) ? FString(" (Not build)") : FString());
+	}
+
+	return FString();
+}
+
+HLODOutliner::FTreeItemID HLODOutliner::FLODActorItem::GetID()
+{
+	return ID;
+}
+
+FSlateColor HLODOutliner::FLODActorItem::GetTint() const
+{
+	ALODActor* LODActorPtr = LODActor.Get();
+	if (LODActorPtr)
+	{
+		return LODActorPtr->IsDirty() ? FLinearColor(1.0f, 0.6f, 0.6f) : FLinearColor(1.0f, 1.0f, 1.0f);
+	}
+
+	return FLinearColor(1.0f, 1.0f, 1.0f);
+}
+
+FText HLODOutliner::FLODActorItem::GetNumTrianglesAsText() const
+{
+	if (LODActor.IsValid())
+	{
+		return FText::FromString(FString::FromInt(LODActor->GetNumTriangles()));
+	}
+	else
+	{
+		return FText::FromString(TEXT("Not available"));
+	}
+}
+
+void HLODOutliner::FLODActorItem::PopulateDragDropPayload(FDragDropPayload& Payload) const
+{
+	ALODActor* ActorPtr = LODActor.Get();
+	if (ActorPtr)
+	{
+		if (!Payload.LODActors)
+		{
+			Payload.LODActors = TArray<TWeakObjectPtr<AActor>>();
+		}
+		Payload.LODActors->Add(LODActor);
+	}
+}
+
+HLODOutliner::FDragValidationInfo HLODOutliner::FLODActorItem::ValidateDrop(FDragDropPayload& DraggedObjects) const
+{
+	if (Parent.IsValid())
+	{
+		if (Parent.Pin()->GetTreeItemType() == ITreeItem::HierarchicalLODActor)
+		{
+			return FDragValidationInfo(FHLODOutlinerDragDropOp::ToolTip_Incompatible, LOCTEXT("CannotAddToChildCluster", "Cannot add to child cluster"));
+		}
+	}
+
+	FLODActorDropTarget Target(LODActor.Get());
+	return Target.ValidateDrop(DraggedObjects);
+}
+
+void HLODOutliner::FLODActorItem::OnDrop(FDragDropPayload& DraggedObjects, const FDragValidationInfo& ValidationInfo, TSharedRef<SWidget> DroppedOnWidget)
+{
+	FLODActorDropTarget Target(LODActor.Get());
+	Target.OnDrop(DraggedObjects, ValidationInfo, DroppedOnWidget);
+}
+
+HLODOutliner::FDragValidationInfo HLODOutliner::FLODActorDropTarget::ValidateDrop(FDragDropPayload& DraggedObjects) const
+{
+	if (DraggedObjects.StaticMeshActors.IsSet() && DraggedObjects.StaticMeshActors->Num() > 0)
+	{
+		if (DraggedObjects.StaticMeshActors->Num() > 0 && DraggedObjects.LODActors->Num() == 0)
+		{			
+			if (DraggedObjects.bSceneOutliner == false)
+			{
+				bool bContaining = false;
+
+				// Check if this StaticMesh Actor is not already inside this cluster
+				for (int32 ActorIndex = 0; ActorIndex < DraggedObjects.StaticMeshActors->Num(); ++ActorIndex)
+				{
+					if (LODActor->SubActors.Contains(DraggedObjects.StaticMeshActors.GetValue()[ActorIndex]))
+					{
+						bContaining = true;
+						break;
+					}
+				}
+
+				if (!bContaining)
+				{
+					if (DraggedObjects.StaticMeshActors->Num() > 1)
+					{
+						return FDragValidationInfo(FHLODOutlinerDragDropOp::ToolTip_MultipleSelection_CompatibleMoveToCluster, LOCTEXT("MoveMultipleToCluster", "Move Actors to Cluster"));
+					}
+					else
+					{
+						return FDragValidationInfo(FHLODOutlinerDragDropOp::ToolTip_CompatibleMoveToCluster, LOCTEXT("MoveToCluster", "Move Actor to Cluster"));
+					}
+				}
+				else
+				{
+					return FDragValidationInfo(FHLODOutlinerDragDropOp::ToolTip_Incompatible, LOCTEXT("AlreadyInCluster", "Cannot add to existing cluster"));
+				}
+			}
+			else
+			{
+				if (DraggedObjects.StaticMeshActors->Num() > 1)
+				{
+					return FDragValidationInfo(FHLODOutlinerDragDropOp::ToolTip_MultipleSelection_CompatibleAddToCluster, LOCTEXT("AddMultipleToCluster", "Add Actors to Cluster"));
+				}
+				else
+				{
+					return FDragValidationInfo(FHLODOutlinerDragDropOp::ToolTip_CompatibleAddToCluster, LOCTEXT("AddToCluster", "Add Actor to Cluster"));
+				}
+
+			}			
+		}
+
+		if (DraggedObjects.bSceneOutliner)
+		{
+			return FDragValidationInfo(FHLODOutlinerDragDropOp::ToolTip_Incompatible, LOCTEXT("AlreadyInHLOD", "Actor is already in one of the Hierarchical LOD clusters"));
+		}
+	}
+	else if (DraggedObjects.LODActors.IsSet() && DraggedObjects.LODActors->Num() > 0)
+	{
+		// Only valid if dragging inside the HLOD outliner
+		if (!DraggedObjects.bSceneOutliner)
+		{		
+			bool bValidForMerge = true;
+			bool bValidForChilding = true;
+			int32 FirstLODLevel = -1;
+
+			for (auto Actor : DraggedObjects.LODActors.GetValue())
+			{
+				ALODActor* InLODActor = Cast<ALODActor>(Actor.Get());			
+
+				if (InLODActor)
+				{
+					// If dragged onto self or already containing LODActor early out
+					if (InLODActor == LODActor.Get() || LODActor->SubActors.Contains(InLODActor))
+					{
+						bValidForMerge = false;
+						bValidForChilding = false;
+						break;
+					}
+
+					// Check in case of multiple selected LODActor items to make sure all of them come from the same LOD level
+					if (FirstLODLevel == -1)
+					{
+						FirstLODLevel = InLODActor->LODLevel;
+					}
+
+					if (InLODActor->LODLevel != LODActor->LODLevel)
+					{
+						bValidForMerge = false;
+
+						if (InLODActor->LODLevel != FirstLODLevel)
+						{
+							bValidForChilding = false;
+						}
+					}						
+				}
+			}
+
+			if (bValidForMerge)
+			{
+				return FDragValidationInfo((DraggedObjects.LODActors->Num() == 1) ? FHLODOutlinerDragDropOp::ToolTip_CompatibleMergeCluster : FHLODOutlinerDragDropOp::ToolTip_MultipleSelection_CompatibleMergeClusters, LOCTEXT("MergeHLODClusters", "Merge cluster(s)"));				
+			}
+			else if (bValidForChilding)
+			{
+				return FDragValidationInfo((DraggedObjects.LODActors->Num() == 1) ? FHLODOutlinerDragDropOp::ToolTip_CompatibleChildCluster : FHLODOutlinerDragDropOp::ToolTip_MultipleSelection_CompatibleChildClusters, LOCTEXT("ChildHLODClusters", "Child cluster(s)"));
+			}
+			else
+			{
+				return FDragValidationInfo(FHLODOutlinerDragDropOp::ToolTip_Incompatible, LOCTEXT("InvalidOperation", "Invalid Operation"));//"Cannot merge clusters from different Hierarchical LOD levels"));
+			}
+		}
+	}
+
+	return FDragValidationInfo(FHLODOutlinerDragDropOp::ToolTip_Incompatible, LOCTEXT("NotImplemented", "Not implemented"));
+}
+
+void HLODOutliner::FLODActorDropTarget::OnDrop(FDragDropPayload& DraggedObjects, const FDragValidationInfo& ValidationInfo, TSharedRef<SWidget> DroppedOnWidget)
+{
+	AActor* DropActor = LODActor.Get();
+	if (!DropActor)
+	{
+		return;
+	}
+	
+	auto& DraggedStaticMeshActors = DraggedObjects.StaticMeshActors.GetValue();
+	auto& DraggedLODActors = DraggedObjects.LODActors.GetValue();
+	if (ValidationInfo.TooltipType == FHLODOutlinerDragDropOp::ToolTip_CompatibleMoveToCluster || ValidationInfo.TooltipType == FHLODOutlinerDragDropOp::ToolTip_MultipleSelection_CompatibleMoveToCluster)
+	{
+		for (int32 ActorIndex = 0; ActorIndex < DraggedStaticMeshActors.Num(); ++ActorIndex)
+		{
+			auto Actor = DraggedStaticMeshActors[ActorIndex];
+			MoveToCluster(Actor.Get(), LODActor.Get());
+		}
+	}
+	else if (ValidationInfo.TooltipType == FHLODOutlinerDragDropOp::ToolTip_CompatibleAddToCluster || ValidationInfo.TooltipType == FHLODOutlinerDragDropOp::ToolTip_MultipleSelection_CompatibleAddToCluster)
+	{
+		for (int32 ActorIndex = 0; ActorIndex < DraggedStaticMeshActors.Num(); ++ActorIndex)
+		{
+			auto Actor = DraggedStaticMeshActors[ActorIndex];
+			AddToCluster(Actor.Get(), LODActor.Get());
+		}
+	}
+	else if (ValidationInfo.TooltipType == FHLODOutlinerDragDropOp::ToolTip_CompatibleMergeCluster || ValidationInfo.TooltipType == FHLODOutlinerDragDropOp::ToolTip_MultipleSelection_CompatibleMergeClusters)
+	{
+		for (int32 ActorIndex = 0; ActorIndex < DraggedLODActors.Num(); ++ActorIndex)
+		{
+			ALODActor* InLODActor = Cast<ALODActor>(DraggedLODActors[ActorIndex].Get());
+			MergeCluster(InLODActor);
+		}
+	}
+	else if (ValidationInfo.TooltipType == FHLODOutlinerDragDropOp::ToolTip_CompatibleChildCluster || ValidationInfo.TooltipType == FHLODOutlinerDragDropOp::ToolTip_MultipleSelection_CompatibleChildClusters)
+	{
+		for (int32 ActorIndex = 0; ActorIndex < DraggedLODActors.Num(); ++ActorIndex)
+		{
+			auto Actor = DraggedLODActors[ActorIndex];
+			MoveToCluster(Actor.Get(), LODActor.Get());
+		}
+	}
+}
+
+void HLODOutliner::FLODActorDropTarget::MoveToCluster(AActor* InActor, ALODActor* NewParentActor)
+{
+	const FScopedTransaction Transaction(LOCTEXT("UndoAction_MoveActorBetweenClusters", "Move Actor between Clusters"));
+	InActor->Modify();
+	ALODActor* CurrentParentActor = HierarchicalLODUtils::GetParentLODActor(InActor);
+	if (CurrentParentActor)
+	{
+		CurrentParentActor->Modify();
+		CurrentParentActor->RemoveSubActor(InActor);
+	}
+	
+	NewParentActor->Modify();
+	NewParentActor->AddSubActor(InActor);
+	
+	GEngine->BroadcastHLODActorMoved(InActor, NewParentActor);
+}
+
+void HLODOutliner::FLODActorDropTarget::AddToCluster(AActor* InActor, ALODActor* NewParentActor)
+{
+	const FScopedTransaction Transaction(LOCTEXT("UndoAction_AddActorToCluster", "Add Actor To Cluster"));
+	NewParentActor->Modify();
+	InActor->Modify();
+
+	NewParentActor->AddSubActor(InActor);
+	GEngine->BroadcastHLODActorAdded(InActor, NewParentActor);
+}
+
+void HLODOutliner::FLODActorDropTarget::MergeCluster(ALODActor* ToMergeActor)
+{
+	const FScopedTransaction Transaction(LOCTEXT("UndoAction_MergeClusters", "Merge Clusters"));
+	ToMergeActor->Modify();
+	LODActor->Modify();
+
+	while (ToMergeActor->SubActors.Num())
+	{
+		AActor* SubActor = ToMergeActor->SubActors.Last();
+		MoveToCluster(SubActor, LODActor.Get());
+	}
+}
\ No newline at end of file
diff --git a/Engine/Source/Editor/HierarchicalLODOutliner/Private/LODActorItem.h b/Engine/Source/Editor/HierarchicalLODOutliner/Private/LODActorItem.h
new file mode 100644
--- /dev/null
+++ b/Engine/Source/Editor/HierarchicalLODOutliner/Private/LODActorItem.h
@@ -0,0 +1,64 @@
+// Copyright 1998-2015 Epic Games, Inc. All Rights Reserved.
+#pragma once
+
+#include "ITreeItem.h"
+
+class ALODActor;
+class AActor;
+
+namespace HLODOutliner
+{
+	/** Helper class to manage moving arbitrary data onto an actor */
+	struct FLODActorDropTarget : IDropTarget
+	{
+		/** The actor this tree item is associated with. */
+		TWeakObjectPtr<ALODActor> LODActor;
+
+		/** Construct this object out of an Actor */
+		FLODActorDropTarget(ALODActor* InLODActor) : LODActor(InLODActor) {}
+
+	public:
+		/** Called to test whether the specified payload can be dropped onto this tree item */
+		virtual FDragValidationInfo ValidateDrop(FDragDropPayload& DraggedObjects) const override;
+
+		/** Called to drop the specified objects on this item. Only called if ValidateDrop() allows. */
+		virtual void OnDrop(FDragDropPayload& DraggedObjects, const FDragValidationInfo& ValidationInfo, TSharedRef<SWidget> DroppedOnWidget) override;
+	private:
+		/** Move a static mesh actor to a different LODActor (cluster) */
+		void MoveToCluster(AActor* InActor, ALODActor* NewParentActor);
+
+		/** Add a static mesh actor to a LODActor (cluster) */
+		void AddToCluster(AActor* InActor, ALODActor* NewParentActor);
+
+		/** Merges ToMergeActor in the represented Cluster/LODActor */
+		void MergeCluster(ALODActor* ToMergeActor);
+	};
+
+	struct FLODActorItem : ITreeItem
+	{
+		mutable TWeakObjectPtr<ALODActor> LODActor;
+		mutable FTreeItemID ID;
+
+		FLODActorItem(ALODActor* InLODActor);
+
+		// Begin ITreeItem interface.
+		virtual bool CanInteract() const override;
+		virtual void GenerateContextMenu(FMenuBuilder& MenuBuilder, SHLODOutliner& Outliner) override;
+		virtual FString GetDisplayString() const override;
+		virtual FSlateColor GetTint() const override;
+		virtual FTreeItemID GetID() override;
+		// End ITreeItem interface.
+
+		/** Returns the number of triangles for the represented LODActor in FText form */
+		FText GetNumTrianglesAsText() const;
+
+		/** Populate the specified drag/drop payload with any relevant information for this type */
+		virtual void PopulateDragDropPayload(FDragDropPayload& Payload) const override;
+
+		/** Called to test whether the specified payload can be dropped onto this tree item */
+		virtual FDragValidationInfo ValidateDrop(FDragDropPayload& DraggedObjects) const override;
+
+		/** Called to drop the specified objects on this item. Only called if ValidateDrop() allows. */
+		virtual void OnDrop(FDragDropPayload& DraggedObjects, const FDragValidationInfo& ValidationInfo, TSharedRef<SWidget> DroppedOnWidget) override;
+	};
+};
\ No newline at end of file
diff --git a/Engine/Source/Editor/HierarchicalLODOutliner/Private/LODLevelItem.cpp b/Engine/Source/Editor/HierarchicalLODOutliner/Private/LODLevelItem.cpp
new file mode 100644
--- /dev/null
+++ b/Engine/Source/Editor/HierarchicalLODOutliner/Private/LODLevelItem.cpp
@@ -0,0 +1,131 @@
+// Copyright 1998-2015 Epic Games, Inc. All Rights Reserved.
+
+#include "HierarchicalLODOutlinerPrivatePCH.h"
+#include "HLODOutliner.h"
+#include "LevelEditor.h"
+#include "LODLevelItem.h"
+#include "HierarchicalLODUtils.h"
+#include "ScopedTransaction.h"
+#include "SlateBasics.h"
+
+
+#define LOCTEXT_NAMESPACE "LODLevelItem"
+
+HLODOutliner::FLODLevelItem::FLODLevelItem(const uint32 InLODIndex)
+	: LODLevelIndex(InLODIndex), ID(nullptr)
+{
+	Type = ITreeItem::HierarchicalLODLevel;
+	ID.SetCachedHash(GetTypeHash(FString("LODLevel - ") + FString::FromInt(LODLevelIndex)));
+}
+
+bool HLODOutliner::FLODLevelItem::CanInteract() const
+{
+	return true;
+}
+
+void HLODOutliner::FLODLevelItem::GenerateContextMenu(FMenuBuilder& MenuBuilder, SHLODOutliner& Outliner)
+{
+
+}
+
+FString HLODOutliner::FLODLevelItem::GetDisplayString() const
+{
+	return FString("LODLevel - ") + FString::FromInt(LODLevelIndex);
+}
+
+HLODOutliner::FTreeItemID HLODOutliner::FLODLevelItem::GetID()
+{
+	return ID;
+}
+
+void HLODOutliner::FLODLevelItem::PopulateDragDropPayload(FDragDropPayload& Payload) const
+{
+
+}
+
+HLODOutliner::FDragValidationInfo HLODOutliner::FLODLevelItem::ValidateDrop(FDragDropPayload& DraggedObjects) const
+{
+	FLODLevelDropTarget Target(LODLevelIndex);
+	return Target.ValidateDrop(DraggedObjects);
+}
+
+void HLODOutliner::FLODLevelItem::OnDrop(FDragDropPayload& DraggedObjects, const FDragValidationInfo& ValidationInfo, TSharedRef<SWidget> DroppedOnWidget)
+{
+	FLODLevelDropTarget Target(LODLevelIndex);
+	return Target.OnDrop(DraggedObjects, ValidationInfo, DroppedOnWidget);
+}
+
+HLODOutliner::FDragValidationInfo HLODOutliner::FLODLevelDropTarget::ValidateDrop(FDragDropPayload& DraggedObjects) const
+{
+	if (DraggedObjects.StaticMeshActors.IsSet() && DraggedObjects.StaticMeshActors->Num() > 0)
+	{
+		const int32 NumStaticMeshActors = DraggedObjects.StaticMeshActors->Num();		
+		return FDragValidationInfo(FHLODOutlinerDragDropOp::ToolTip_CompatibleNewCluster, LOCTEXT("CreateNewCluster", "Create new Cluster"));
+	}
+	else if (DraggedObjects.LODActors.IsSet() && DraggedObjects.LODActors->Num() > 0)
+	{
+		const int32 NumLODActors = DraggedObjects.LODActors->Num();
+
+		if (NumLODActors > 1)
+		{
+			return FDragValidationInfo(FHLODOutlinerDragDropOp::ToolTip_MultipleSelection_CompatibleNewCluster, LOCTEXT("CreateNewCluster", "Create new Cluster"));
+		}
+	}
+
+	return FDragValidationInfo(FHLODOutlinerDragDropOp::ToolTip_Incompatible, LOCTEXT("NotImplemented", "Not implemented"));
+}
+
+void HLODOutliner::FLODLevelDropTarget::OnDrop(FDragDropPayload& DraggedObjects, const FDragValidationInfo& ValidationInfo, TSharedRef<SWidget> DroppedOnWidget)
+{
+	if (ValidationInfo.TooltipType == FHLODOutlinerDragDropOp::ToolTip_CompatibleNewCluster || ValidationInfo.TooltipType == FHLODOutlinerDragDropOp::ToolTip_MultipleSelection_CompatibleNewCluster)
+	{
+		CreateNewCluster(DraggedObjects);
+	}
+}
+
+void HLODOutliner::FLODLevelDropTarget::CreateNewCluster(FDragDropPayload &DraggedObjects)
+{	
+	// Change this to pass world into OnDrop?
+	UWorld* World = nullptr;
+	if (DraggedObjects.StaticMeshActors.IsSet() && DraggedObjects.StaticMeshActors->Num() > 0)
+	{
+		World = DraggedObjects.StaticMeshActors.GetValue()[0]->GetWorld();
+	}
+	else if (DraggedObjects.LODActors.IsSet() && DraggedObjects.LODActors->Num() > 0)
+	{
+		World = DraggedObjects.LODActors.GetValue()[0]->GetWorld();
+	}
+
+	const FScopedTransaction Transaction(LOCTEXT("UndoAction_CreateNewCluster", "Create new Cluster"));
+	World->Modify();
+	
+	ALODActor* NewCluster = HierarchicalLODUtils::CreateNewClusterActor(World, LODLevelIndex);
+
+	for (TWeakObjectPtr<AActor> StaticMeshActor : DraggedObjects.StaticMeshActors.GetValue())
+	{
+		AActor* InActor = StaticMeshActor.Get();
+		ALODActor* CurrentParentActor = HierarchicalLODUtils::GetParentLODActor(InActor);
+		if (CurrentParentActor)
+		{
+			CurrentParentActor->RemoveSubActor(InActor);
+		}
+
+		NewCluster->AddSubActor(InActor);
+	}
+
+	for (TWeakObjectPtr<AActor> LODActor : DraggedObjects.LODActors.GetValue())
+	{
+		AActor* InActor = LODActor.Get();
+		ALODActor* CurrentParentActor = HierarchicalLODUtils::GetParentLODActor(InActor);
+		if (CurrentParentActor)
+		{
+			CurrentParentActor->RemoveSubActor(InActor);
+		}
+
+		NewCluster->AddSubActor(InActor);
+	}
+}
+
+#undef LOCTEXT_NAMESPACE
+
+
diff --git a/Engine/Source/Editor/HierarchicalLODOutliner/Private/LODLevelItem.h b/Engine/Source/Editor/HierarchicalLODOutliner/Private/LODLevelItem.h
new file mode 100644
--- /dev/null
+++ b/Engine/Source/Editor/HierarchicalLODOutliner/Private/LODLevelItem.h
@@ -0,0 +1,55 @@
+// Copyright 1998-2015 Epic Games, Inc. All Rights Reserved.
+#pragma once
+
+#include "ITreeItem.h"
+
+namespace HLODOutliner
+{
+	/** Helper class to manage moving abritrary data onto an actor */
+	struct FLODLevelDropTarget : IDropTarget
+	{
+		/** The actor this tree item is associated with. */
+		uint32 LODLevelIndex;
+
+		/** Construct this object out of an Actor */
+		FLODLevelDropTarget(uint32 InLODIndex) : LODLevelIndex(InLODIndex) {}
+
+	public:
+
+		/** Called to test whether the specified payload can be dropped onto this tree item */
+		virtual FDragValidationInfo ValidateDrop(FDragDropPayload& DraggedObjects) const override;
+
+		/** Called to drop the specified objects on this item. Only called if ValidateDrop() allows. */
+		virtual void OnDrop(FDragDropPayload& DraggedObjects, const FDragValidationInfo& ValidationInfo, TSharedRef<SWidget> DroppedOnWidget) override;
+
+		/** Creating a new cluster (LODActor) with given objects/actors */
+		void CreateNewCluster(FDragDropPayload &DraggedObjects);
+	};
+
+	struct FLODLevelItem : ITreeItem
+	{
+		FLODLevelItem(const uint32 InLODIndex);
+
+		// Begin ITreeItem interface.
+		virtual bool CanInteract() const override;
+		virtual void GenerateContextMenu(FMenuBuilder& MenuBuilder, SHLODOutliner& Outliner) override;
+		virtual FString GetDisplayString() const override;
+		virtual FTreeItemID GetID() override;
+		// Begin ITreeItem interface.
+
+		/** Populate the specified drag/drop payload with any relevant information for this type */
+		virtual void PopulateDragDropPayload(FDragDropPayload& Payload) const override;
+
+		/** Called to test whether the specified payload can be dropped onto this tree item */
+		virtual FDragValidationInfo ValidateDrop(FDragDropPayload& DraggedObjects) const override;
+
+		/** Called to drop the specified objects on this item. Only called if ValidateDrop() allows. */
+		virtual void OnDrop(FDragDropPayload& DraggedObjects, const FDragValidationInfo& ValidationInfo, TSharedRef<SWidget> DroppedOnWidget) override;
+		
+		/** LOD index of this level */
+		uint32 LODLevelIndex;
+
+		/** TreeItem ID */
+		FTreeItemID ID;
+	};
+};
diff --git a/Engine/Source/Editor/HierarchicalLODOutliner/Private/StaticMeshActorItem.cpp b/Engine/Source/Editor/HierarchicalLODOutliner/Private/StaticMeshActorItem.cpp
new file mode 100644
--- /dev/null
+++ b/Engine/Source/Editor/HierarchicalLODOutliner/Private/StaticMeshActorItem.cpp
@@ -0,0 +1,74 @@
+// Copyright 1998-2015 Epic Games, Inc. All Rights Reserved.
+
+#include "HierarchicalLODOutlinerPrivatePCH.h"
+#include "StaticMeshActorItem.h"
+#include "HLODOutliner.h"
+#include "LevelEditor.h"
+#include "SlateBasics.h"
+
+#define LOCTEXT_NAMESPACE "StaticMeshActorItem"
+
+HLODOutliner::FStaticMeshActorItem::FStaticMeshActorItem(AActor* InStaticMeshActor)
+	: StaticMeshActor(InStaticMeshActor), ID(InStaticMeshActor)
+{
+	Type = ITreeItem::StaticMeshActor;
+}
+
+bool HLODOutliner::FStaticMeshActorItem::CanInteract() const
+{
+	return true;
+}
+
+void HLODOutliner::FStaticMeshActorItem::GenerateContextMenu(FMenuBuilder& MenuBuilder, SHLODOutliner& Outliner)
+{
+	auto SharedOutliner = StaticCastSharedRef<SHLODOutliner>(Outliner.AsShared());
+	MenuBuilder.AddMenuEntry(LOCTEXT("RemoveSMActorFromCluster", "Remove from cluster"), FText(), FSlateIcon(), FUIAction(FExecuteAction::CreateSP(&Outliner, &SHLODOutliner::RemoveStaticMeshActorFromCluster, AsShared())));
+}
+
+FString HLODOutliner::FStaticMeshActorItem::GetDisplayString() const
+{
+	return StaticMeshActor->GetFName().GetPlainNameString();
+}
+
+HLODOutliner::FTreeItemID HLODOutliner::FStaticMeshActorItem::GetID()
+{
+	return ID;
+}
+
+void HLODOutliner::FStaticMeshActorItem::PopulateDragDropPayload(FDragDropPayload& Payload) const
+{
+	AActor* ActorPtr = StaticMeshActor.Get();
+	if (ActorPtr)
+	{
+		if (!Payload.StaticMeshActors)
+		{
+			Payload.StaticMeshActors = TArray<TWeakObjectPtr<AActor>>();
+		}
+		Payload.StaticMeshActors->Add(StaticMeshActor);		
+	}
+}
+
+HLODOutliner::FDragValidationInfo HLODOutliner::FStaticMeshActorItem::ValidateDrop(FDragDropPayload& DraggedObjects) const
+{
+	FStaticMeshActorDropTarget Target(StaticMeshActor.Get());
+	return Target.ValidateDrop(DraggedObjects);
+}
+
+void HLODOutliner::FStaticMeshActorItem::OnDrop(FDragDropPayload& DraggedObjects, const FDragValidationInfo& ValidationInfo, TSharedRef<SWidget> DroppedOnWidget)
+{
+	FStaticMeshActorDropTarget Target(StaticMeshActor.Get());
+	Target.OnDrop(DraggedObjects, ValidationInfo, DroppedOnWidget);
+}
+
+
+HLODOutliner::FDragValidationInfo HLODOutliner::FStaticMeshActorDropTarget::ValidateDrop(FDragDropPayload& DraggedObjects) const
+{
+	return FDragValidationInfo(FHLODOutlinerDragDropOp::ToolTip_Incompatible, LOCTEXT("NotImplemented", "Not implemented"));
+}
+
+void HLODOutliner::FStaticMeshActorDropTarget::OnDrop(FDragDropPayload& DraggedObjects, const FDragValidationInfo& ValidationInfo, TSharedRef<SWidget> DroppedOnWidget)
+{
+
+}
+
+#undef LOCTEXT_NAMESPACE // "StaticMeshActorItem"
diff --git a/Engine/Source/Editor/HierarchicalLODOutliner/Private/StaticMeshActorItem.h b/Engine/Source/Editor/HierarchicalLODOutliner/Private/StaticMeshActorItem.h
new file mode 100644
--- /dev/null
+++ b/Engine/Source/Editor/HierarchicalLODOutliner/Private/StaticMeshActorItem.h
@@ -0,0 +1,52 @@
+// Copyright 1998-2015 Epic Games, Inc. All Rights Reserved.
+#pragma once
+
+#include "ITreeItem.h"
+
+
+namespace HLODOutliner
+{
+	class SHLODOutliner;
+
+	/** Helper class to manage moving abritrary data onto an actor */
+	struct FStaticMeshActorDropTarget : IDropTarget
+	{
+		/** The actor this tree item is associated with. */
+		TWeakObjectPtr<AActor> Actor;
+
+		/** Construct this object out of an Actor */
+		FStaticMeshActorDropTarget(AActor* InActor) : Actor(InActor) {}
+	public:
+		/** Called to test whether the specified payload can be dropped onto this tree item */
+		virtual FDragValidationInfo ValidateDrop(FDragDropPayload& DraggedObjects) const override;
+
+		/** Called to drop the specified objects on this item. Only called if ValidateDrop() allows. */
+		virtual void OnDrop(FDragDropPayload& DraggedObjects, const FDragValidationInfo& ValidationInfo, TSharedRef<SWidget> DroppedOnWidget) override;
+	};
+
+	struct FStaticMeshActorItem : ITreeItem
+	{
+		/** Represented StaticMeshActor */
+		mutable TWeakObjectPtr<AActor> StaticMeshActor;
+		/** TreeItem's ID */
+		mutable FTreeItemID ID;
+
+		FStaticMeshActorItem(AActor* InStaticMeshActor);
+
+		// Begin ITreeItem interface.
+		virtual bool CanInteract() const override;
+		virtual void GenerateContextMenu(FMenuBuilder& MenuBuilder, SHLODOutliner& Outliner) override;
+		virtual FString GetDisplayString() const override;
+		virtual FTreeItemID GetID() override;
+		// End ITreeItem interface.
+
+		/** Populate the specified drag/drop payload with any relevant information for this type */
+		virtual void PopulateDragDropPayload(FDragDropPayload& Payload) const override;
+
+		/** Called to test whether the specified payload can be dropped onto this tree item */
+		virtual FDragValidationInfo ValidateDrop(FDragDropPayload& DraggedObjects) const override;
+
+		/** Called to drop the specified objects on this item. Only called if ValidateDrop() allows. */
+		virtual void OnDrop(FDragDropPayload& DraggedObjects, const FDragValidationInfo& ValidationInfo, TSharedRef<SWidget> DroppedOnWidget) override;
+	};
+};
diff --git a/Engine/Source/Editor/HierarchicalLODOutliner/Private/TreeItemID.h b/Engine/Source/Editor/HierarchicalLODOutliner/Private/TreeItemID.h
new file mode 100644
--- /dev/null
+++ b/Engine/Source/Editor/HierarchicalLODOutliner/Private/TreeItemID.h
@@ -0,0 +1,61 @@
+// Copyright 1998-2015 Epic Games, Inc. All Rights Reserved.
+#pragma once
+
+#include "ObjectKey.h"
+
+namespace HLODOutliner
+{
+	struct FTreeItemID
+	{
+	public:
+		FTreeItemID(const UObject* InObject)
+		{
+			ObjectKey = FObjectKey(InObject);
+			CachedHash = GetTypeHash(ObjectKey);
+		}
+
+		/** Copy construction / assignment */
+		FTreeItemID(const FTreeItemID& Other)
+		{
+			*this = Other;
+		}
+
+		/** Move construction / assignment */
+		FTreeItemID(FTreeItemID&& Other)
+		{
+			*this = MoveTemp(Other);
+		}
+		FTreeItemID& operator=(FTreeItemID&& Other)
+		{
+			FMemory::Memswap(this, &Other, sizeof(FTreeItemID));
+			return *this;
+		}
+
+		~FTreeItemID()
+		{
+			CachedHash = 0;
+		}
+
+		friend bool operator==(const FTreeItemID& One, const FTreeItemID& Other)
+		{
+			return One.CachedHash == Other.CachedHash;
+		}
+		friend bool operator!=(const FTreeItemID& One, const FTreeItemID& Other)
+		{
+			return One.CachedHash != Other.CachedHash;
+		}
+
+		friend uint32 GetTypeHash(const FTreeItemID& ItemID)
+		{
+			return ItemID.CachedHash;
+		}
+
+		void SetCachedHash(const uint32 InHash)
+		{
+			CachedHash = InHash;
+		}
+	private:
+		FObjectKey ObjectKey;
+		uint32 CachedHash;
+	};
+};
diff --git a/Engine/Source/Editor/HierarchicalLODOutliner/Public/HierarchicalLODOutlinerModule.h b/Engine/Source/Editor/HierarchicalLODOutliner/Public/HierarchicalLODOutlinerModule.h
new file mode 100644
--- /dev/null
+++ b/Engine/Source/Editor/HierarchicalLODOutliner/Public/HierarchicalLODOutlinerModule.h
@@ -0,0 +1,23 @@
+// Copyright 1998-2015 Epic Games, Inc. All Rights Reserved.
+
+#include "ModuleInterface.h"
+
+/**
+* The module holding all of the UI related pieces for HLOD Outliner
+*/
+class FHierarchicalLODOutlinerModule : public IModuleInterface
+{
+public:
+	/**
+	* Called right after the module DLL has been loaded and the module object has been created
+	*/
+	virtual void StartupModule();
+
+	/**
+	* Called before the module is unloaded, right before the module object is destroyed.
+	*/
+	virtual void ShutdownModule();
+
+	/** Creates the HLOD Outliner widget */
+	virtual TSharedRef<SWidget> CreateHLODOutlinerWidget();
+};
\ No newline at end of file
diff --git a/Engine/Source/Editor/LevelEditor/LevelEditor.Build.cs b/Engine/Source/Editor/LevelEditor/LevelEditor.Build.cs
--- a/Engine/Source/Editor/LevelEditor/LevelEditor.Build.cs
+++ b/Engine/Source/Editor/LevelEditor/LevelEditor.Build.cs
@@ -23,7 +23,8 @@ public class LevelEditor : ModuleRules
 			new string[] {
 				"Settings",
 				"UserFeedback",
-				"IntroTutorials"
+				"IntroTutorials",
+                "HierarchicalLODOutliner"
 			}
 		);
 
@@ -83,7 +84,8 @@ public class LevelEditor : ModuleRules
 				"UserFeedback",
 				"ReferenceViewer",
 				"SizeMap",
-                "IntroTutorials"
+                "IntroTutorials",
+                "HierarchicalLODOutliner"
 			}
 		);
 	}
diff --git a/Engine/Source/Editor/LevelEditor/Private/LevelEditor.cpp b/Engine/Source/Editor/LevelEditor/Private/LevelEditor.cpp
--- a/Engine/Source/Editor/LevelEditor/Private/LevelEditor.cpp
+++ b/Engine/Source/Editor/LevelEditor/Private/LevelEditor.cpp
@@ -1049,6 +1049,11 @@ void FLevelEditorModule::BindGlobalLevelEditorCommands()
 		FExecuteAction::CreateStatic( &FLevelEditorActionCallbacks::ExecuteExecCommand, FString( TEXT("ACTOR SELECT MATCHINGSTATICMESH ALLCLASSES") ) )
 		);
 
+	ActionList.MapAction(
+		Commands.SelectOwningHierarchicalLODCluster,
+		FExecuteAction::CreateStatic(&FLevelEditorActionCallbacks::OnSelectOwningHLODCluster)
+		);
+
 	ActionList.MapAction(
 		Commands.SelectSkeletalMeshesOfSameClass,
 		FExecuteAction::CreateStatic( &FLevelEditorActionCallbacks::ExecuteExecCommand, FString( TEXT("ACTOR SELECT MATCHINGSKELETALMESH") ) )
diff --git a/Engine/Source/Editor/LevelEditor/Private/LevelEditorActions.cpp b/Engine/Source/Editor/LevelEditor/Private/LevelEditorActions.cpp
--- a/Engine/Source/Editor/LevelEditor/Private/LevelEditorActions.cpp
+++ b/Engine/Source/Editor/LevelEditor/Private/LevelEditorActions.cpp
@@ -64,6 +64,7 @@
 #include "Editor/Persona/Public/AnimationRecorder.h"
 #include "Editor/KismetWidgets/Public/CreateBlueprintFromActorDialog.h"
 #include "EditorProjectSettings.h"
+#include "HierarchicalLODUtils.h"
 
 DEFINE_LOG_CATEGORY_STATIC(LevelEditorActions, Log, All);
 
@@ -1592,7 +1593,22 @@ void FLevelEditorActionCallbacks::OnSelectAllActorsControlledByMatinee()
 	GEditor->SelectAllActorsControlledByMatinee();
 }
 
-void FLevelEditorActionCallbacks::OnSelectMatineeActor( AMatineeActor * ActorToSelect )
+void FLevelEditorActionCallbacks::OnSelectOwningHLODCluster()
+{
+	if (GEditor->GetSelectedActorCount() > 0)
+	{
+		AActor* Actor = Cast<AActor>(GEditor->GetSelectedActors()->GetSelectedObject(0));
+		ALODActor* ParentActor = HierarchicalLODUtils::GetParentLODActor(Actor);
+		if (Actor && ParentActor)
+		{
+			GEditor->SelectNone(false, true);
+			GEditor->SelectActor(ParentActor, true, false);
+			GEditor->NoteSelectionChange();
+		}
+	}
+}
+
+void FLevelEditorActionCallbacks::OnSelectMatineeActor(AMatineeActor * ActorToSelect)
 {
 	GEditor->SelectNone( false, true );
 	GEditor->SelectActor(ActorToSelect, true, false, true);
@@ -2922,6 +2938,7 @@ void FLevelEditorCommands::RegisterCommands()
 	UI_COMMAND( SelectComponentOwnerActor, "Select Component Owner", "Select the actor that owns the currently selected component(s)", EUserInterfaceActionType::Button, FInputChord() );
 	UI_COMMAND( SelectRelevantLights, "Select Relevant Lights", "Select all lights relevant to the current selection", EUserInterfaceActionType::Button, FInputChord() ); 
 	UI_COMMAND( SelectStaticMeshesOfSameClass, "Select All Using Selected Static Meshes (Selected Actor Types)", "Selects all actors with the same static mesh and actor class as the selection", EUserInterfaceActionType::Button, FInputChord() ); 
+	UI_COMMAND( SelectOwningHierarchicalLODCluster, "Select Owning Hierarchical LOD cluster Using Selected Static Mesh (Selected Actor Types)", "Select Owning Hierarchical LOD cluster for the selected actor", EUserInterfaceActionType::Button, FInputChord());
 	UI_COMMAND( SelectStaticMeshesAllClasses, "Select All Using Selected Static Meshes (All Actor Types)", "Selects all actors with the same static mesh as the selection", EUserInterfaceActionType::Button, FInputChord( EModifierKey::Shift, EKeys::E ) ); 
 	UI_COMMAND( SelectSkeletalMeshesOfSameClass, "Select All Using Selected Skeletal Meshes (Selected Actor Types)", "Selects all actors with the same skeletal mesh and actor class as the selection", EUserInterfaceActionType::Button, FInputChord() ); 
 	UI_COMMAND( SelectSkeletalMeshesAllClasses, "Select All Using Selected Skeletal Meshes (All Actor Types)", "Selects all actors with the same skeletal mesh as the selection", EUserInterfaceActionType::Button, FInputChord() ); 
diff --git a/Engine/Source/Editor/LevelEditor/Private/LevelEditorContextMenu.cpp b/Engine/Source/Editor/LevelEditor/Private/LevelEditorContextMenu.cpp
--- a/Engine/Source/Editor/LevelEditor/Private/LevelEditorContextMenu.cpp
+++ b/Engine/Source/Editor/LevelEditor/Private/LevelEditorContextMenu.cpp
@@ -570,6 +570,15 @@ void FLevelEditorContextMenuImpl::FillSelectActorMenu( FMenuBuilder& MenuBuilder
 				MenuBuilder.AddMenuEntry( FLevelEditorCommands::Get().SelectStaticMeshesAllClasses, NAME_None, LOCTEXT("SelectStaticMeshesAllClasses_Menu", "Select Matching (All Classes)") );
 			}
 			MenuBuilder.EndSection();
+
+			if (SelectionInfo.NumSelected == 1)
+			{
+				MenuBuilder.BeginSection("SelectHLODCluster", LOCTEXT("SelectHLODClusterHeading", "Hierachical LODs"));
+				{
+					MenuBuilder.AddMenuEntry(FLevelEditorCommands::Get().SelectOwningHierarchicalLODCluster, NAME_None, LOCTEXT("SelectOwningHierarchicalLODCluster_Menu", "Select Owning HierarchicalLODCluster"));
+				}
+				MenuBuilder.EndSection();
+			}			
 		}
 
 		if( SelectionInfo.bHavePawn || SelectionInfo.bHaveSkeletalMesh )
diff --git a/Engine/Source/Editor/LevelEditor/Private/SLevelEditor.cpp b/Engine/Source/Editor/LevelEditor/Private/SLevelEditor.cpp
--- a/Engine/Source/Editor/LevelEditor/Private/SLevelEditor.cpp
+++ b/Engine/Source/Editor/LevelEditor/Private/SLevelEditor.cpp
@@ -38,6 +38,8 @@
 #include "ScopedTransaction.h"
 #include "GameFramework/WorldSettings.h"
 
+#include "HierarchicalLODOutlinerModule.h"
+
 
 static const FName LevelEditorBuildAndSubmitTab("LevelEditorBuildAndSubmit");
 static const FName LevelEditorStatsViewerTab("LevelEditorStatsViewer");
@@ -703,6 +705,21 @@ TSharedRef<SDockTab> SLevelEditor::SpawnLevelEditorTab( const FSpawnTabArgs& Arg
 					]
 				];
 	}
+	else if (TabIdentifier == TEXT("LevelEditorHierarchicalLODOutliner"))
+	{
+		FHierarchicalLODOutlinerModule& HLODModule = FModuleManager::LoadModuleChecked<FHierarchicalLODOutlinerModule>("HierarchicalLODOutliner");
+		return SNew(SDockTab)
+			.Icon(FEditorStyle::GetBrush("LevelEditor.Tabs.HLOD"))
+			.Label(NSLOCTEXT("LevelEditor", "HLODTabTitle", "Hierarchical LOD Outliner"))
+			[
+				SNew(SBorder)
+				.Padding(0)
+				.BorderImage(FEditorStyle::GetBrush("ToolPanel.GroupBorder"))				
+				[
+					HLODModule.CreateHLODOutlinerWidget()
+				]
+			];
+	}
 	else if( TabIdentifier == WorldBrowserHierarchyTab )
 	{
 		FWorldBrowserModule& WorldBrowserModule = FModuleManager::LoadModuleChecked<FWorldBrowserModule>( "WorldBrowser" );
@@ -1035,6 +1052,15 @@ TSharedRef<SWidget> SLevelEditor::RestoreContentArea( const TSharedRef<SDockTab>
 				.SetGroup( MenuStructure.GetLevelEditorCategory() )
 				.SetIcon( LayersIcon );
 		}
+
+		{
+			const FSlateIcon LayersIcon(FEditorStyle::GetStyleSetName(), "LevelEditor.Tabs.HLOD");
+			LevelEditorTabManager->RegisterTabSpawner("LevelEditorHierarchicalLODOutliner", FOnSpawnTab::CreateSP<SLevelEditor, FName, FString>(this, &SLevelEditor::SpawnLevelEditorTab, FName("LevelEditorHierarchicalLODOutliner"), FString()))
+				.SetDisplayName(NSLOCTEXT("LevelEditorTabs", "LevelEditorHierarchicalLODOutliner", "Hierarchical LOD Outliner"))
+				.SetTooltipText(NSLOCTEXT("LevelEditorTabs", "LevelEditorHierarchicalLODOutlinerTooltipText", "Open the Hierarchical LOD Outliner."))
+				.SetGroup(MenuStructure.GetLevelEditorCategory())
+				.SetIcon(LayersIcon);
+		}
 		
 		{
 			LevelEditorTabManager->RegisterTabSpawner( WorldBrowserHierarchyTab, FOnSpawnTab::CreateSP<SLevelEditor, FName, FString>(this, &SLevelEditor::SpawnLevelEditorTab, WorldBrowserHierarchyTab, FString()) )
diff --git a/Engine/Source/Editor/LevelEditor/Public/LevelEditorActions.h b/Engine/Source/Editor/LevelEditor/Public/LevelEditorActions.h
--- a/Engine/Source/Editor/LevelEditor/Public/LevelEditorActions.h
+++ b/Engine/Source/Editor/LevelEditor/Public/LevelEditorActions.h
@@ -350,6 +350,9 @@ public:
 	/** Selects all actors using the same static mesh(es) and same actor class as the current selection */
 	TSharedPtr< FUICommandInfo > SelectStaticMeshesAllClasses;
 
+	/** Selects the HLOD cluster (ALODActor), if available, that has this actor as one of its SubActors */
+	TSharedPtr< FUICommandInfo > SelectOwningHierarchicalLODCluster;
+
 	/** Selects all actors using the same skeletal mesh(es) as the current selection */
 	TSharedPtr< FUICommandInfo > SelectSkeletalMeshesOfSameClass;
 
@@ -928,6 +931,11 @@ public:
 	 */
 	static void OnSelectAllActorsControlledByMatinee();
 	
+	/**
+	* Called when selecting an Actor's (if available) owning HLOD cluster
+	*/
+	static void OnSelectOwningHLODCluster();
+	
 	/**
 	 * Called to change bsp surface alignment
 	 *
diff --git a/Engine/Source/Editor/UnrealEd/Private/HierarchicalLOD.cpp b/Engine/Source/Editor/UnrealEd/Private/HierarchicalLOD.cpp
--- a/Engine/Source/Editor/UnrealEd/Private/HierarchicalLOD.cpp
+++ b/Engine/Source/Editor/UnrealEd/Private/HierarchicalLOD.cpp
@@ -13,6 +13,8 @@
 #include "Engine/LODActor.h"
 #include "Editor/UnrealEd/Classes/Factories/Factory.h"
 #include "ObjectTools.h"
+#include "MeshUtilities.h"
+#include "HierarchicalLODUtils.h"
 #endif // WITH_EDITOR
 
 #include "GameFramework/WorldSettings.h"
@@ -20,6 +22,7 @@
 
 #include "HierarchicalLODVolume.h"
 
+
 DEFINE_LOG_CATEGORY_STATIC(LogLODGenerator, Warning, All);
 
 #define LOCTEXT_NAMESPACE "HierarchicalLOD"
@@ -41,6 +44,10 @@ FHierarchicalLODBuilder::FHierarchicalLODBuilder()
 void FHierarchicalLODBuilder::Build()
 {
 	check (World)
+
+	AWorldSettings* WorldSetting = World->GetWorldSettings();
+	BuildLODLevelSettings = WorldSetting->HierarchicalLODSetup;
+
 	const TArray<ULevel*>& Levels = World->GetLevels();
 	for (const auto& LevelIter : Levels)
 	{
@@ -51,6 +58,9 @@ void FHierarchicalLODBuilder::Build()
 void FHierarchicalLODBuilder::PreviewBuild()
 {
 	check(World);
+	AWorldSettings* WorldSetting = World->GetWorldSettings();
+	BuildLODLevelSettings = WorldSetting->HierarchicalLODSetup;
+
 	const TArray<ULevel*>& Levels = World->GetLevels();
 	for (const auto& LevelIter : Levels)
 	{
@@ -65,55 +75,27 @@ void FHierarchicalLODBuilder::BuildClusters(ULevel* InLevel, const bool bCreateM
 	// I'm using stack mem within this scope of the function
 	// so we need this
 	FMemMark Mark(FMemStack::Get());
-
-	// you still have to delete all objects just in case they had it and didn't want it anymore
-	TArray<UObject*> AssetsToDelete;
-	for (int32 ActorId=InLevel->Actors.Num()-1; ActorId >= 0; --ActorId)
-	{
-		ALODActor* LodActor = Cast<ALODActor>(InLevel->Actors[ActorId]);
-		if (LodActor)
-		{
-			for (auto& Asset: LodActor->SubObjects)
-			{
-				// @TOOD: This is not permanent fix
-				if (Asset)
-				{
-					AssetsToDelete.Add(Asset);
-				}
-			}
-			World->DestroyActor(LodActor);
-		}
-	}
-
-	ULevel::BuildStreamingData(InLevel->OwningWorld, InLevel);
-
-	for (auto& Asset : AssetsToDelete)
-	{
-		Asset->MarkPendingKill();
-		ObjectTools::DeleteSingleObject(Asset, false);
-	}
 	
-	// garbage collect
-	CollectGarbage( GARBAGE_COLLECTION_KEEPFLAGS, true );
+	DeleteLODActors(InLevel, !bCreateMeshes);
+	
 
 	// only build if it's enabled
-	if(InLevel->GetWorld()->GetWorldSettings()->bEnableHierarchicalLODSystem && InLevel->GetWorld()->GetWorldSettings()->HierarchicalLODSetup.Num() != 0)
+	if (InLevel->GetWorld()->GetWorldSettings()->bEnableHierarchicalLODSystem && BuildLODLevelSettings.Num() > 0)
 	{
 		// Handle HierachicalLOD volumes first
 		HandleHLODVolumes(InLevel);
 
-		AWorldSettings* WorldSetting = InLevel->GetWorld()->GetWorldSettings();		
-		const int32 TotalNumLOD = InLevel->GetWorld()->GetWorldSettings()->HierarchicalLODSetup.Num();
+		const int32 TotalNumLOD = BuildLODLevelSettings.Num();
 		for(int32 LODId=0; LODId<TotalNumLOD; ++LODId)
 		{
 			
 			// we use meter for bound. Otherwise it's very easy to get to overflow and have problem with filling ratio because
 			// bound is too huge
-			const float DesiredBoundRadius = WorldSetting->HierarchicalLODSetup[LODId].DesiredBoundRadius * CM_TO_METER;
-			const float DesiredFillingRatio = WorldSetting->HierarchicalLODSetup[LODId].DesiredFillingPercentage * 0.01f;
+			const float DesiredBoundRadius = BuildLODLevelSettings[LODId].DesiredBoundRadius * CM_TO_METER;
+			const float DesiredFillingRatio = BuildLODLevelSettings[LODId].DesiredFillingPercentage * 0.01f;
 			ensure(DesiredFillingRatio!=0.f);
 			const float HighestCost = FMath::Pow(DesiredBoundRadius, 3) / (DesiredFillingRatio);
-			const int32 MinNumActors = WorldSetting->HierarchicalLODSetup[LODId].MinNumberOfActorsToBuild;
+			const int32 MinNumActors = BuildLODLevelSettings[LODId].MinNumberOfActorsToBuild;
 			check (MinNumActors > 0);
 			// test parameter I was playing with to cull adding to the array
 			// intialization can have too many elements, decided to cull
@@ -132,7 +114,7 @@ void FHierarchicalLODBuilder::BuildClusters(ULevel* InLevel, const bool bCreateM
 				SlowTask.MakeDialog();
 
 				// initialize Clusters
-				InitializeClusters(InLevel, LODId, HighestCost*CullMultiplier);
+				InitializeClusters(InLevel, LODId, HighestCost*CullMultiplier, !bCreateMeshes);
 
 				// move a half way - I know we can do this better but as of now this is small progress
 				SlowTask.EnterProgressFrame(50);
@@ -151,7 +133,7 @@ void FHierarchicalLODBuilder::BuildClusters(ULevel* InLevel, const bool bCreateM
 		FMessageLog MapCheck("MapCheck");
 		MapCheck.Warning()
 			->AddToken(FUObjectToken::Create(InLevel->GetWorld()->GetWorldSettings()))
-			->AddToken(FTextToken::Create(LOCTEXT("MapCheck_Message_HLODSystemNotEnabled", "Hierarchical LOD System is disabled or no HLOD level settings available, unable to build LOD actors.")))
+			->AddToken(FTextToken::Create(LOCTEXT("MapCheck_Message_HLODSystemNotEnabled", "Hierarchical LOD System is disabled, unable to build LOD actors.")))
 			->AddToken(FMapErrorToken::Create(FMapErrors::HLODSystemNotEnabled));
 	}
 
@@ -160,7 +142,7 @@ void FHierarchicalLODBuilder::BuildClusters(ULevel* InLevel, const bool bCreateM
 	Clusters.Shrink();
 }
 
-void FHierarchicalLODBuilder::InitializeClusters(ULevel* InLevel, const int32 LODIdx, float CullCost)
+void FHierarchicalLODBuilder::InitializeClusters(ULevel* InLevel, const int32 LODIdx, float CullCost, const bool bPreviewBuild)
 {
 	SCOPE_LOG_TIME(TEXT("STAT_HLOD_InitializeClusters"), nullptr);
 	if (InLevel->Actors.Num() > 0)
@@ -173,18 +155,24 @@ void FHierarchicalLODBuilder::InitializeClusters(ULevel* InLevel, const int32 LO
 			for (int32 ActorId = 0; ActorId < InLevel->Actors.Num(); ++ActorId)
 			{
 				AActor* Actor = InLevel->Actors[ActorId];
-				if (ShouldGenerateCluster(Actor))
+				if (ShouldGenerateCluster(Actor, bPreviewBuild))
 				{
 					// Check whether or not this actor falls within a HierarchicalLODVolume, if so add to the Volume's cluster and exclude from normal process
 					bool bAdded = false;
 					for (auto& Cluster : HLODVolumeClusters)
 					{
 						if (Cluster.Key->EncompassesPoint(Actor->GetActorLocation(), 0.0f, nullptr))
-						{
-							FLODCluster ActorCluster(Actor);
-							Cluster.Value += ActorCluster;
-							bAdded = true;
-							break;
+						{			
+							FBox BoundingBox = Actor->GetComponentsBoundingBox(true);
+							FBox VolumeBox = Cluster.Key->GetComponentsBoundingBox(true);
+
+							if (VolumeBox.IsInside(BoundingBox))
+							{
+								FLODCluster ActorCluster(Actor);
+								Cluster.Value += ActorCluster;
+								bAdded = true;
+								break;
+							}							
 						}
 					}
 
@@ -229,12 +217,13 @@ void FHierarchicalLODBuilder::InitializeClusters(ULevel* InLevel, const int32 LO
 					if (Actor->IsA(ALODActor::StaticClass()))
 					{
 						ALODActor* LODActor = CastChecked<ALODActor>(Actor);
-						if (LODActor->LODLevel == LODIdx)
+
+						if (LODActor->LODLevel == LODIdx && ((bPreviewBuild && !LODActor->GetStaticMeshComponent()->StaticMesh) || !bPreviewBuild ))
 						{
 							Actors.Add(Actor);
 						}
 					}
-					else if (ShouldGenerateCluster(Actor))
+					else if (ShouldGenerateCluster(Actor, bPreviewBuild))
 					{
 						Actors.Add(Actor);
 					}
@@ -292,29 +281,52 @@ void FHierarchicalLODBuilder::HandleHLODVolumes(ULevel* InLevel)
 		if (Actor && Actor->IsA( AHierarchicalLODVolume::StaticClass()))
 		{
 			// Came across a HLOD volume			
-			FLODCluster& NewCluster = HLODVolumeClusters.Add(CastChecked<AHierarchicalLODVolume>(Actor));			
+			FLODCluster& NewCluster = HLODVolumeClusters.Add(CastChecked<AHierarchicalLODVolume>(Actor));
+
+			FVector Origin, Extent;
+			Actor->GetActorBounds(false, Origin, Extent);
+			NewCluster.Bound = FSphere(Origin * CM_TO_METER, Extent.Size() * CM_TO_METER);
+
+
+			// calculate new filling factor
+			NewCluster.FillingFactor = 1.f;
+			NewCluster.ClusterCost = FMath::Pow(NewCluster.Bound.W, 3) / NewCluster.FillingFactor;
 		}
 	}
 }
 
-bool FHierarchicalLODBuilder::ShouldGenerateCluster(AActor* Actor)
+bool FHierarchicalLODBuilder::ShouldGenerateCluster(AActor* Actor, const bool bPreviewBuild)
 {
 	if (!Actor)
 	{
 		return false;
 	}
 
+	if (Actor->bHidden)
+	{
+		return false;
+	}
+
 	if (!Actor->bEnableAutoLODGeneration)
 	{
 		return false;
 	}
 
+	ALODActor* LODActor = Cast<ALODActor>(Actor);
+	if (bPreviewBuild && LODActor)
+	{
+		if (LODActor->GetStaticMeshComponent()->StaticMesh)
+		{
+			return false;
+		}
+	}
+
 	FVector Origin, Extent;
 	Actor->GetActorBounds(false, Origin, Extent);
 	if (Extent.SizeSquared() <= 0.1)
 	{
 		return false;
-	}
+	}	
 
 	// for now only consider staticmesh - I don't think skel mesh would work with simplygon merge right now @fixme
 	TArray<UStaticMeshComponent*> Components;
@@ -328,7 +340,17 @@ bool FHierarchicalLODBuilder::ShouldGenerateCluster(AActor* Actor)
 	{
 		for (auto& ComponentIter : Components)
 		{
+			
 			if (ComponentIter->GetLODParentPrimitive())
+			{
+				auto ParentActor = CastChecked<ALODActor>(ComponentIter->GetLODParentPrimitive()->GetOwner());
+				if (!ParentActor->GetStaticMeshComponent()->StaticMesh || !bPreviewBuild)
+				{
+					return false;
+				}
+			}
+
+			if (ComponentIter->bHiddenInGame)
 			{
 				return false;
 			}
@@ -345,9 +367,127 @@ bool FHierarchicalLODBuilder::ShouldGenerateCluster(AActor* Actor)
 	return (ValidComponentCount > 0);
 }
 
+void FHierarchicalLODBuilder::ClearHLODs()
+{
+	for (auto& Level : World->GetLevels())
+	{
+		DeleteLODActors(Level, false);
+	}
+}
+
+void FHierarchicalLODBuilder::ClearPreviewBuild()
+{
+	for (auto& Level : World->GetLevels())
+	{
+		DeleteLODActors(Level, true);
+	}
+}
+
+void FHierarchicalLODBuilder::BuildMeshesForLODActors()
+{
+	AWorldSettings* WorldSetting = World->GetWorldSettings();
+	BuildLODLevelSettings = WorldSetting->HierarchicalLODSetup;
+
+	const TArray<ULevel*>& Levels = World->GetLevels();
+	for (const auto& LevelIter : Levels)
+	{
+		FScopedSlowTask SlowTask(100, (LOCTEXT("HierarchicalLOD_BuildLODActorMeshes", "Building LODActor meshes")));
+		SlowTask.MakeDialog();
+
+		TArray<TArray<ALODActor*>> LODLevelActors;
+		LODLevelActors.AddDefaulted(BuildLODLevelSettings.Num());
+
+		if (LevelIter->Actors.Num() > 0)
+		{
+			UPackage* AssetsOuter = LevelIter->GetOutermost();
+
+			if (AssetsOuter)
+			{
+				uint32 NumLODActors = 0;
+				for (int32 ActorId = 0; ActorId < LevelIter->Actors.Num(); ++ActorId)
+				{
+					AActor* Actor = LevelIter->Actors[ActorId];
+					if (Actor && Actor->IsA<ALODActor>())
+					{
+						ALODActor* LODActor = CastChecked<ALODActor>(Actor);
+
+						if (LODActor->IsDirty() && LODActor->SubActors.Num() > 1)
+						{
+							LODLevelActors[LODActor->LODLevel - 1].Add(LODActor);
+							NumLODActors++;
+						}
+					}
+				}		
+			
+				const int32 NumLODLevels = LODLevelActors.Num();
+				for (int32 LODIndex = 0; LODIndex < NumLODLevels; ++LODIndex)
+				{
+					int32 CurrentLODLevel = LODIndex;
+					TArray<ALODActor*>& LODLevel = LODLevelActors[CurrentLODLevel];
+					for (ALODActor* Actor : LODLevel)
+					{
+						HierarchicalLODUtils::BuildStaticMeshForLODActor(Actor, AssetsOuter, BuildLODLevelSettings[CurrentLODLevel], CurrentLODLevel);
+						SlowTask.EnterProgressFrame(100.0f / (float)NumLODActors);
+					}
+				}
+				
+				
+			}
+		}		
+	}
+}
+
+void FHierarchicalLODBuilder::DeleteLODActors(ULevel* InLevel, const bool bPreviewOnly)
+{
+	// you still have to delete all objects just in case they had it and didn't want it anymore
+	TArray<UObject*> AssetsToDelete;
+	for (int32 ActorId = InLevel->Actors.Num() - 1; ActorId >= 0; --ActorId)
+	{
+		ALODActor* LodActor = Cast<ALODActor>(InLevel->Actors[ActorId]);
+		if (LodActor && ( ( LodActor->GetStaticMeshComponent()->StaticMesh == nullptr && bPreviewOnly) || !bPreviewOnly ))
+		{
+			for (auto& Asset : LodActor->SubObjects)
+			{
+				// @TOOD: This is not permanent fix
+				if (Asset)
+				{
+					AssetsToDelete.Add(Asset);
+				}
+			}
+			World->DestroyActor(LodActor);
+		}
+	}
+
+	ULevel::BuildStreamingData(InLevel->OwningWorld, InLevel);
+
+	for (auto& Asset : AssetsToDelete)
+	{
+		Asset->MarkPendingKill();
+		ObjectTools::DeleteSingleObject(Asset, false);
+	}
+
+	// garbage collect
+	CollectGarbage(GARBAGE_COLLECTION_KEEPFLAGS, true);
+}
+
+void FHierarchicalLODBuilder::BuildMeshForLODActor(ALODActor* LODActor, const uint32 LODLevel)
+{
+	FScopedSlowTask SlowTask(100, (LOCTEXT("HierarchicalLOD_BuildLODActorMeshes", "Building LODActor Mesh")));
+	SlowTask.MakeDialog();
+
+	AWorldSettings* WorldSetting = World->GetWorldSettings();
+	BuildLODLevelSettings = WorldSetting->HierarchicalLODSetup;
+	UPackage* AssetsOuter = LODActor->GetLevel()->GetOutermost();
+
+	const bool bResult = HierarchicalLODUtils::BuildStaticMeshForLODActor(LODActor, AssetsOuter, BuildLODLevelSettings[LODLevel], LODLevel);
+	check(bResult);
+
+	SlowTask.EnterProgressFrame(100.0f);
+}
+
 void FHierarchicalLODBuilder::MergeClustersAndBuildActors(ULevel* InLevel, const int32 LODIdx, float HighestCost, int32 MinNumActors, const bool bCreateMeshes)
 {	
-	if (Clusters.Num() > 0)
+	if (Clusters.Num() > 0 || HLODVolumeClusters.Num() > 0)
 	{
 		FString LevelName = FPackageName::GetShortName(InLevel->GetOutermost()->GetName());
 		FFormatNamedArguments Arguments;
@@ -364,6 +504,7 @@ void FHierarchicalLODBuilder::MergeClustersAndBuildActors(ULevel* InLevel, const
 
 			for(int32 Iteration=0; Iteration<TotalIteration; ++Iteration)
 			{
+				bool bChanged = false;
 				// now we have minimum Clusters
 				for(int32 ClusterId=0; ClusterId < TotalCluster; ++ClusterId)
 				{
@@ -397,11 +538,14 @@ void FHierarchicalLODBuilder::MergeClustersAndBuildActors(ULevel* InLevel, const
 										MergedCluster = NewCluster;
 										// now this cluster is invalid
 										Cluster.Invalidate();
+
+										bChanged = true;
 										break;
 									}
 									else
 									{
 										Cluster -= MergedCluster;
+										bChanged = true;
 									}
 								}
 							}
@@ -410,10 +554,14 @@ void FHierarchicalLODBuilder::MergeClustersAndBuildActors(ULevel* InLevel, const
 						UE_LOG(LogLODGenerator, Verbose, TEXT("Processed(%s): %0.2f {%s}"), Cluster.IsValid()? TEXT("Valid"):TEXT("Invalid"), Cluster.GetCost(), *Cluster.ToString());
 					}
 				}
+
+				if (bChanged == false)
+				{
+					break;
+				}					
 			}
 		}
 
-
 		if (LODIdx == 0)
 		{
 			for (auto& Cluster : HLODVolumeClusters)
@@ -422,11 +570,7 @@ void FHierarchicalLODBuilder::MergeClustersAndBuildActors(ULevel* InLevel, const
 			}
 		}
 		
-
-		// debug flag, so that I can just see clustered data since no visualization in the editor yet
-		//static bool bBuildActor=true;
-
-		//if (bBuildActors)
+		
 		{
 			SCOPE_LOG_TIME(TEXT("HLOD_BuildActors"), nullptr);
 			// print data
diff --git a/Engine/Source/Editor/UnrealEd/Private/LODCluster.cpp b/Engine/Source/Editor/UnrealEd/Private/LODCluster.cpp
--- a/Engine/Source/Editor/UnrealEd/Private/LODCluster.cpp
+++ b/Engine/Source/Editor/UnrealEd/Private/LODCluster.cpp
@@ -11,6 +11,7 @@
 #include "Editor/UnrealEd/Classes/Factories/Factory.h"
 #include "MeshUtilities.h"
 #include "ObjectTools.h"
+#include "HierarchicalLODUtils.h"
 #endif // WITH_EDITOR
 
 #include "GameFramework/WorldSettings.h"
@@ -22,6 +23,8 @@
 #define CM_TO_METER		0.01f
 #define METER_TO_CM		100.0f
 
+
+
 /** Utility function to calculate overlap of two spheres */
 const float CalculateOverlap(const FSphere& ASphere, const float AFillingFactor, const FSphere& BSphere, const float BFillingFactor)
 {
@@ -31,17 +34,17 @@ const float CalculateOverlap(const FSphere& ASphere, const float AFillingFactor,
 		return 0.f;
 	}
 
-	if (ASphere.IsInside(BSphere, 1.f))
+	if (ASphere.IsInside(BSphere))
 	{
 		return ASphere.GetVolume();
 	}
 
-	if(BSphere.IsInside(ASphere, 1.f))
+	if(BSphere.IsInside(ASphere))
 	{
 		return BSphere.GetVolume();
 	}
 
-	if (ASphere.Equals(BSphere, 1.f))
+	if (ASphere.Equals(BSphere))
 	{
 		return ASphere.GetVolume();
 	}
@@ -90,16 +93,17 @@ const float CalculateFillingFactor(const FSphere& ASphere, const float AFillingF
 	// http://deim.urv.cat/~rivi/pub/3d/icra04b.pdf
 	// cost is calculated based on r^3 / filling factor
 	// since it subtract by AFillingFactor * 1/2 overlap volume + BfillingFactor * 1/2 overlap volume
-	return FMath::Max(0.f, (AFillingFactor * ASphere.GetVolume() + BFillingFactor * BSphere.GetVolume() - OverlapVolume) / UnionSphere.GetVolume());
+	return FMath::Max(0.0f, (AFillingFactor * ASphere.GetVolume() + BFillingFactor * BSphere.GetVolume() - OverlapVolume) / UnionSphere.GetVolume());
 }
 
 FLODCluster::FLODCluster(const FLODCluster& Other)
 : Actors(Other.Actors)
 , Bound(Other.Bound)
 , FillingFactor(Other.FillingFactor)
-, ClusterCost( FMath::Pow(Bound.W, 3) / FillingFactor )
+, ClusterCost(Other.ClusterCost)
 , bValid(Other.bValid)
 {
+	
 }
 
 FLODCluster::FLODCluster(AActor* Actor1)
@@ -128,8 +132,7 @@ FLODCluster::FLODCluster()
 : Bound(ForceInit)
 , bValid(false)
 {
-	// calculate new filling factor
-	FillingFactor = 1.f;
+	FillingFactor = 1.0f;
 	ClusterCost = FMath::Pow(Bound.W, 3) / FillingFactor;
 }
 
@@ -139,26 +142,8 @@ FSphere FLODCluster::AddActor(AActor* NewActor)
 	ensure (Actors.Contains(NewActor) == false);
 	Actors.Add(NewActor);
 	FVector Origin, Extent;
-#ifdef WITH_EDITORONLY_DATA
-	if (NewActor->IsA<ALODActor>())
-	{
-		ALODActor* LODActor = CastChecked<ALODActor>(NewActor);
 
-		if (LODActor->IsPreviewActorInstance())
-		{
-			Origin = LODActor->GetDrawSphereComponent()->GetComponentLocation();
-			Extent = FVector(LODActor->GetDrawSphereComponent()->GetScaledSphereRadius());
-		}
-		else
-		{
-			NewActor->GetActorBounds(false, Origin, Extent);
-		}
-	}
-	else
-#endif // WITH_EDITORONLY_DATA
-	{
-		NewActor->GetActorBounds(false, Origin, Extent);
-	}
+	NewActor->GetActorBounds(false, Origin, Extent);
 
 	// scale 0.01 (change to meter from centimeter) QQ Extens.GetSize	
 	FSphere NewBound = FSphere(Origin*CM_TO_METER, Extent.Size()*CM_TO_METER);
@@ -210,7 +195,9 @@ void FLODCluster::MergeClusters(const FLODCluster& Other)
 	// have to recalculate filling factor and bound based on cluster data
 	FillingFactor = CalculateFillingFactor(Bound, FillingFactor, Other.Bound, Other.FillingFactor);
 	Bound += Other.Bound;	
+
 	ClusterCost = FMath::Pow(Bound.W, 3) / FillingFactor;
+	
 
 	for (auto& Actor: Other.Actors)
 	{
@@ -299,7 +286,7 @@ void FLODCluster::BuildActor(ULevel* InLevel, const int32 LODIdx, const bool bCr
 				
 				if (Actor->IsA<ALODActor>())
 				{
-					ExtractStaticMeshComponentsFromLODActor(Actor, Components);
+					HierarchicalLODUtils::ExtractStaticMeshComponentsFromLODActor(Actor, Components);
 				}
 				else
 				{
@@ -363,55 +350,32 @@ void FLODCluster::BuildActor(ULevel* InLevel, const int32 LODIdx, const bool bCr
 					UWorld* LevelWorld = Cast<UWorld>(InLevel->GetOuter());
 
 					check (LevelWorld);
+
+					FTransform Transform;
+					Transform.SetLocation(OutProxyLocation);					
 										
 					// create LODActors using the current Actors
-					ALODActor* NewActor = Cast<ALODActor>(LevelWorld->SpawnActor(ALODActor::StaticClass(), &OutProxyLocation, &FRotator::ZeroRotator));
+					ALODActor* NewActor = nullptr;
+										
+					NewActor = LevelWorld->SpawnActor<ALODActor>(ALODActor::StaticClass(), Transform);				
 					NewActor->SubObjects = OutAssets;
-					NewActor->SubActors = Actors;
-					NewActor->LODLevel = LODIdx+1; 					
+					NewActor->LODLevel = LODIdx+1;
 					NewActor->LODDrawDistance = DrawDistance;
 					NewActor->SetStaticMesh( MainMesh );
-					NewActor->SetFolderPath("/HLODActors");											
-
-#if WITH_EDITORONLY_DATA						
-					// Setup sphere drawing visualization
-					UDrawSphereComponent* SphereComponent = NewActor->GetDrawSphereComponent();
-					SphereComponent->SetSphereRadius(Bound.W * METER_TO_CM);
-					SphereComponent->SetWorldLocation(Bound.Center * METER_TO_CM);					
-					SphereComponent->ShapeColor = Colours[LODIdx % 7];
-					SphereComponent->MinDrawDistance = DrawDistance;
-					SphereComponent->LDMaxDrawDistance = NextDrawDistance;
-					SphereComponent->CachedMaxDrawDistance = NextDrawDistance;	
-#endif // WITH_EDITORONLY_DATA						
 					// now set as parent
-					for (auto& Actor : Actors)
+					for(auto& Actor : Actors)
 					{
-						Actor->SetLODParent(NewActor->GetStaticMeshComponent(), DrawDistance);
-					}					
+						NewActor->AddSubActor(Actor);
+					}
+
+					// Mark dirty according to whether or not this is a preview build
+					NewActor->SetIsDirty(!bCreateMeshes);
 				}
 			}
 		}
 	}
 }
 
-void FLODCluster::ExtractStaticMeshComponentsFromLODActor(AActor* Actor, TArray<UStaticMeshComponent*> &InOutComponents)
-{
-	ALODActor* LODActor = CastChecked<ALODActor>(Actor);
-	for (auto& ChildActor : LODActor->SubActors)
-	{
-		TArray<UStaticMeshComponent*> ChildComponents;
-		if (ChildActor->IsA<ALODActor>())
-		{
-			ExtractStaticMeshComponentsFromLODActor(ChildActor, ChildComponents);
-		}
-		else
-		{
-			ChildActor->GetComponents<UStaticMeshComponent>(ChildComponents);
-		}
-		
-		InOutComponents.Append(ChildComponents);
-	}
-}
 
 bool FLODCluster::Contains(FLODCluster& Other) const
 {
diff --git a/Engine/Source/Editor/UnrealEd/Public/HierarchicalLOD.h b/Engine/Source/Editor/UnrealEd/Public/HierarchicalLOD.h
--- a/Engine/Source/Editor/UnrealEd/Public/HierarchicalLOD.h
+++ b/Engine/Source/Editor/UnrealEd/Public/HierarchicalLOD.h
@@ -12,6 +12,7 @@ class ULevel;
 class AActor;
 class AHierarchicalLODVolume;
 class UWorld;
+class ALODActor;
 
 /**
  *
@@ -38,6 +39,27 @@ struct UNREALED_API FHierarchicalLODBuilder
 	*/
 	void PreviewBuild();
 
+	/**
+	* Clear all the HLODs and the ALODActors that were created for them
+	*/
+	void ClearHLODs();
+
+	/**
+	* Clear only the ALODActorsPreview 
+	*/
+	void ClearPreviewBuild();
+
+	/** Builds the LOD meshes for all LODActors inside of the World's Levels */
+	void BuildMeshesForLODActors();
+
+	/**
+	* Build a single LOD Actor's mesh
+	*
+	* @param LODActor - LODActor to build mesh for
+	* @param LODLevel - LODLevel to build the mesh for
+	*/
+	void BuildMeshForLODActor(ALODActor* LODActor, const uint32 LODLevel);
+
 private:
 	/**
 	* Builds the clusters (HLODs) for InLevel, and will create the new/merged StaticMeshes if bCreateMeshes is true
@@ -54,7 +76,7 @@ private:
 	* @param LODIdx - LOD index we are building
 	* @param CullCost - Test variable for tweaking HighestCost
 	*/
-	void InitializeClusters(ULevel* InLevel, const int32 LODIdx, float CullCost);
+	void InitializeClusters(ULevel* InLevel, const int32 LODIdx, float CullCost, bool const bPreviewBuild);
 
 	/**
 	* Merges clusters and builds actors for resulting (valid) clusters
@@ -84,8 +106,17 @@ private:
 	* @param Actor - Actor to test
 	* @return bool
 	*/
-	bool ShouldGenerateCluster(AActor* Actor);
+	bool ShouldGenerateCluster(AActor* Actor, const bool bPreviewBuild);
 	
+	/**
+	* Deletes LOD actors from the world	
+	*
+	* @param InLevel - Level to delete the actors from
+	* @param bPreviewOnly - Only delete preview actors
+	* @return void
+	*/
+	void DeleteLODActors(ULevel* InLevel, const bool bPreviewOnly);
+
 	/** Array of LOD Clusters - this is only valid within scope since mem stack allocator */
 	TArray<FLODCluster, TMemStackAllocator<>> Clusters;
 
@@ -94,4 +125,7 @@ private:
 
 	/** Array of LOD clusters created for the HierachicalLODVolumes found within the level */
 	TMap<AHierarchicalLODVolume*, FLODCluster> HLODVolumeClusters;	
+
+	/** Cached array of LODLevel settings */
+	TArray<FHierarchicalSimplification> BuildLODLevelSettings;
 };
diff --git a/Engine/Source/Editor/UnrealEd/Public/HierarchicalLODUtils.h b/Engine/Source/Editor/UnrealEd/Public/HierarchicalLODUtils.h
new file mode 100644
--- /dev/null
+++ b/Engine/Source/Editor/UnrealEd/Public/HierarchicalLODUtils.h
@@ -0,0 +1,333 @@
+// Copyright 1998-2015 Epic Games, Inc. All Rights Reserved.
+#pragma once
+
+#include "Engine/LODActor.h"
+#include "Components/StaticMeshComponent.h"
+#include "Engine/StaticMesh.h"
+#include "MeshUtilities.h"
+
+#include "Asyncwork.h"
+
+#if WITH_EDITOR
+#include "AssetData.h"
+#include "Factories/Factory.h"
+#include "ObjectTools.h"
+#endif // WITH_EDITOR
+
+#define LOCTEXT_NAMESPACE "HierarchicalLODUtils"
+
+namespace HierarchicalLODUtils
+{
+	/**
+	* Recursively retrieves StaticMeshComponents from a LODActor and its child LODActors
+	*
+	* @param Actor - LODActor instance
+	* @param InOutComponents - Will hold the StaticMeshComponents
+	*/
+	static void ExtractStaticMeshComponentsFromLODActor(AActor* Actor, TArray<UStaticMeshComponent*>& InOutComponents)
+	{
+		ALODActor* LODActor = CastChecked<ALODActor>(Actor);
+		for (auto& ChildActor : LODActor->SubActors)
+		{
+			TArray<UStaticMeshComponent*> ChildComponents;
+			if (ChildActor->IsA<ALODActor>())
+			{
+				ExtractStaticMeshComponentsFromLODActor(ChildActor, ChildComponents);
+			}
+			else
+			{
+				ChildActor->GetComponents<UStaticMeshComponent>(ChildComponents);
+			}
+
+			InOutComponents.Append(ChildComponents);
+		}
+	}
+
+
+	/**
+	* Builds a static mesh object for the given LODACtor
+	*
+	* @param LODActor -
+	* @param Outer -
+	* @return UStaticMesh*
+	*/
+	static bool BuildStaticMeshForLODActor(ALODActor* LODActor, UPackage* AssetsOuter, const FHierarchicalSimplification& LODSetup, const uint32 LODIndex)
+	{
+		if (AssetsOuter && LODActor)
+		{
+			if (!LODActor->IsDirty() || LODActor->SubActors.Num() < 2)
+			{
+				return true;
+			}
+
+			// Clean up sub objects (if applicable)
+			for( auto& SubOject : LODActor->SubObjects)
+			{
+				if (SubOject)
+				{
+					SubOject->MarkPendingKill();
+				}				
+			}
+			LODActor->SubObjects.Empty();
+
+			TArray<UStaticMeshComponent*> AllComponents;
+						
+			{
+				FScopedSlowTask SlowTask(LODActor->SubActors.Num(), (LOCTEXT("HierarchicalLODUtils_CollectStaticMeshes", "Collecting Static Meshes for Cluster")));
+				SlowTask.MakeDialog();
+				for (auto& Actor : LODActor->SubActors)
+				{
+					TArray<UStaticMeshComponent*> Components;
+
+					if (Actor->IsA<ALODActor>())
+					{
+						ExtractStaticMeshComponentsFromLODActor(Actor, Components);
+					}
+					else
+					{
+						Actor->GetComponents<UStaticMeshComponent>(Components);
+					}
+
+					// TODO: support instanced static meshes
+					Components.RemoveAll([](UStaticMeshComponent* Val){ return Val->IsA(UInstancedStaticMeshComponent::StaticClass()); });
+
+					AllComponents.Append(Components);
+					SlowTask.EnterProgressFrame(1);
+				}
+			}
+
+			// it shouldn't even have come here if it didn't have any staticmesh
+			if (ensure(AllComponents.Num() > 0))
+			{
+				FScopedSlowTask SlowTask(LODActor->SubActors.Num(), (LOCTEXT("HierarchicalLODUtils_MergingMeshes", "Merging Static Meshes and creating LODActor")));
+				SlowTask.MakeDialog();
+
+				// In case we don't have outer generated assets should have same path as LOD level
+				const FString AssetsPath = AssetsOuter->GetName() + TEXT("/");
+				AActor* FirstActor = LODActor->SubActors[0];
+
+				TArray<UObject*> OutAssets;
+				FVector OutProxyLocation = FVector::ZeroVector;
+				UStaticMesh* MainMesh = nullptr;
+			
+				// Generate proxy mesh and proxy material assets
+				IMeshUtilities& MeshUtilities = FModuleManager::Get().LoadModuleChecked<IMeshUtilities>("MeshUtilities");
+				// should give unique name, so use level + actor name
+				const FString PackageName = FString::Printf(TEXT("LOD_%s"), *FirstActor->GetName());
+				if (MeshUtilities.GetMeshMergingInterface() && LODSetup.bSimplifyMesh)
+				{
+					MeshUtilities.CreateProxyMesh(LODActor->SubActors, LODSetup.ProxySetting, AssetsOuter, PackageName, OutAssets, OutProxyLocation);
+				}
+				else
+				{
+					MeshUtilities.MergeStaticMeshComponents(AllComponents, FirstActor->GetWorld(), LODSetup.MergeSetting, AssetsOuter, PackageName, LODIndex, OutAssets, OutProxyLocation, LODSetup.DrawDistance, true);
+				}
+
+				// we make it private, so it can't be used by outside of map since it's useless, and then remove standalone
+				for (auto& AssetIter : OutAssets)
+				{
+					AssetIter->ClearFlags(RF_Public | RF_Standalone);
+				}
+
+				// set staticmesh
+				for (auto& Asset : OutAssets)
+				{
+					UStaticMesh* StaticMesh = Cast<UStaticMesh>(Asset);
+
+					if (StaticMesh)
+					{
+						MainMesh = StaticMesh;
+					}
+				}
+
+				LODActor->SetStaticMesh(MainMesh);
+				LODActor->SetActorLocation(OutProxyLocation);
+				LODActor->SubObjects = OutAssets;				
+		
+				for (auto& Actor : LODActor->SubActors)
+				{
+					Actor->SetLODParent(LODActor->GetStaticMeshComponent(), LODActor->LODDrawDistance);
+				}
+
+				// Freshly build so mark not dirty
+				LODActor->SetIsDirty(false);
+
+				return true;
+			}
+		}
+		return false;
+	}
+
+
+	static bool ShouldGenerateCluster(AActor* Actor)
+	{
+		if (!Actor)
+		{
+			return false;
+		}
+
+		if (Actor->bHidden)
+		{
+			return false;
+		}
+
+		if (!Actor->bEnableAutoLODGeneration)
+		{
+			return false;
+		}
+
+		ALODActor* LODActor = Cast<ALODActor>(Actor);
+		if (LODActor)
+		{			
+			return false;			
+		}
+
+		FVector Origin, Extent;
+		Actor->GetActorBounds(false, Origin, Extent);
+		if (Extent.SizeSquared() <= 0.1)
+		{
+			return false;
+		}
+
+		// for now only consider staticmesh - I don't think skel mesh would work with simplygon merge right now @fixme
+		TArray<UStaticMeshComponent*> Components;
+		Actor->GetComponents<UStaticMeshComponent>(Components);
+		// TODO: support instanced static meshes
+		Components.RemoveAll([](UStaticMeshComponent* Val){ return Val->IsA(UInstancedStaticMeshComponent::StaticClass()); });
+
+		int32 ValidComponentCount = 0;
+		// now make sure you check parent primitive, so that we don't build for the actor that already has built. 
+		if (Components.Num() > 0)
+		{
+			for (auto& ComponentIter : Components)
+			{
+				if (ComponentIter->GetLODParentPrimitive())
+				{
+					return false;					
+				}
+
+				if (ComponentIter->bHiddenInGame)
+				{
+					return false;
+				}
+
+				// see if we should generate it
+				if (ComponentIter->ShouldGenerateAutoLOD())
+				{
+					++ValidComponentCount;
+					break;
+				}
+			}
+		}
+
+		return (ValidComponentCount > 0);
+	}
+
+	static ALODActor* GetParentLODActor(const AActor* InActor)
+	{
+		ALODActor* ParentActor = nullptr;
+		if (InActor)
+		{
+			TArray<UStaticMeshComponent*> ComponentArray;
+			InActor->GetComponents<UStaticMeshComponent>(ComponentArray);
+			for (auto Component : ComponentArray)
+			{
+				UPrimitiveComponent* ParentComponent = Component->GetLODParentPrimitive();
+				if (ParentComponent)
+				{
+					ParentActor = CastChecked<ALODActor>(ParentComponent->GetOwner());
+					break;
+				}
+			}
+		}
+
+		return ParentActor;		
+	}
+
+	static void DeleteLODActor(ALODActor* InActor)
+	{
+		// Find if it has a parent ALODActor
+		AActor* Actor = CastChecked<AActor>(InActor);
+
+		ALODActor* ParentLOD = GetParentLODActor(Actor);
+		if (ParentLOD)
+		{
+			ParentLOD->RemoveSubActor(Actor);
+		}
+
+		// Clean out sub actors and update their LODParent
+		while (InActor->SubActors.Num())
+		{
+			auto SubActor = InActor->SubActors[0];
+			InActor->RemoveSubActor(SubActor);
+		}
+
+		// you still have to delete all objects just in case they had it and didn't want it anymore
+		TArray<UObject*> AssetsToDelete;
+		
+		AssetsToDelete.Append(InActor->SubObjects);
+		
+		for (auto& Asset : AssetsToDelete)
+		{
+			if (Asset)
+			{
+				Asset->MarkPendingKill();
+				ObjectTools::DeleteSingleObject(Asset, false);
+			}			
+		}
+
+		// Rebuild streaming data to prevent invalid references
+		ULevel::BuildStreamingData(InActor->GetWorld(), InActor->GetLevel());
+
+		// garbage collect
+		CollectGarbage(GARBAGE_COLLECTION_KEEPFLAGS, true);
+	}
+
+	static void DeleteLODActorAssets(ALODActor* InActor)
+	{
+		// you still have to delete all objects just in case they had it and didn't want it anymore
+		TArray<UObject*> AssetsToDelete;
+
+		AssetsToDelete.Append(InActor->SubObjects);
+		InActor->SubObjects.Empty();
+		for (auto& Asset : AssetsToDelete)
+		{
+			if (Asset)
+			{
+				Asset->MarkPendingKill();
+				ObjectTools::DeleteSingleObject(Asset, false);
+			}
+		}
+
+		InActor->GetStaticMeshComponent()->StaticMesh = nullptr;
+
+		// Mark render state dirty
+		InActor->GetStaticMeshComponent()->MarkRenderStateDirty();	
+
+		// Rebuild streaming data to prevent invalid references
+		ULevel::BuildStreamingData(InActor->GetWorld(), InActor->GetLevel());
+		
+		// Garbage collecting
+		CollectGarbage(GARBAGE_COLLECTION_KEEPFLAGS, true);
+	}
+
+	static ALODActor* CreateNewClusterActor(UWorld* InWorld, const int32 InLODLevel)
+	{
+		// Check incoming data
+		check(InWorld != nullptr && InLODLevel >= 0);
+		auto WorldSettings = InWorld->GetWorldSettings();
+		check(WorldSettings->bEnableHierarchicalLODSystem && WorldSettings->HierarchicalLODSetup.Num() > InLODLevel);
+					
+		// Spawn and setup actor
+		ALODActor* NewActor = nullptr;
+		NewActor = InWorld->SpawnActor<ALODActor>(ALODActor::StaticClass(), FTransform());
+		NewActor->LODLevel = InLODLevel + 1;
+		NewActor->LODDrawDistance = WorldSettings->HierarchicalLODSetup[InLODLevel].DrawDistance;
+		NewActor->SetStaticMesh(nullptr);
+
+		return NewActor;		
+	}
+
+};
+
+#undef LOCTEXT_NAMESPACE
\ No newline at end of file
diff --git a/Engine/Source/Editor/UnrealEd/UnrealEd.Build.cs b/Engine/Source/Editor/UnrealEd/UnrealEd.Build.cs
--- a/Engine/Source/Editor/UnrealEd/UnrealEd.Build.cs
+++ b/Engine/Source/Editor/UnrealEd/UnrealEd.Build.cs
@@ -229,7 +229,8 @@ public class UnrealEd : ModuleRules
 				"NiagaraEditor",
                 "Niagara",
                 "VectorVM",
-				"AddContentDialog",
+				"AddContentDialog",                
+                "MeshUtilities"
 			}
 			);
 
diff --git a/Engine/Source/Runtime/Engine/Classes/Engine/Engine.h b/Engine/Source/Runtime/Engine/Classes/Engine/Engine.h
--- a/Engine/Source/Runtime/Engine/Classes/Engine/Engine.h
+++ b/Engine/Source/Runtime/Engine/Classes/Engine/Engine.h
@@ -1588,6 +1588,41 @@ public:
 	/** Called by internal engine systems after a level actor has been requested to be renamed */
 	void BroadcastLevelComponentRequestRename(const UActorComponent* InComponent) { LevelComponentRequestRenameEvent.Broadcast(InComponent); }
 
+	/** Editor-only event triggered when a HLOD Actor is moved between clusters */
+	DECLARE_EVENT_TwoParams(UEngine, FHLODActorMovedEvent, const AActor*, const AActor*);
+	FHLODActorMovedEvent& OnHLODActorMoved() { return HLODActorMovedEvent; }
+
+	/** Called by internal engine systems after a HLOD Actor is moved between clusters */
+	void BroadcastHLODActorMoved(const AActor* InActor, const AActor* ParentActor ) { HLODActorMovedEvent.Broadcast(InActor, ParentActor); }
+
+	/** Editor-only event triggered when a HLOD Actor's mesh is build */
+	DECLARE_EVENT_OneParam(UEngine, FHLODMeshBuildEvent, const class ALODActor*);
+	FHLODMeshBuildEvent& OnHLODMeshBuild() { return HLODMeshBuildEvent; }
+
+	/** Called by internal engine systems after a HLOD Actor's mesh is build */
+	void BroadcastHLODMeshBuild(const class ALODActor* InActor) { HLODMeshBuildEvent.Broadcast(InActor); }
+
+	/** Editor-only event triggered when a HLOD Actor is added to a cluster */
+	DECLARE_EVENT_TwoParams(UEngine, FHLODActorAddedEvent, const AActor*, const AActor*);
+	FHLODActorAddedEvent& OnHLODActorAdded() { return HLODActorAddedEvent; }
+
+	/** Called by internal engine systems after a HLOD Actor is added to a cluster */
+	void BroadcastHLODActorAdded(const AActor* InActor, const AActor* ParentActor) { HLODActorAddedEvent.Broadcast(InActor, ParentActor); }
+
+	/** Editor-only event triggered when a HLOD Actor is marked dirty */
+	DECLARE_EVENT_OneParam(UEngine, FHLODActorMarkedDirtyEvent, class ALODActor*);
+	FHLODActorMarkedDirtyEvent& OnHLODActorMarkedDirty() { return HLODActorMarkedDirtyEvent; }
+
+	/** Called by internal engine systems after a HLOD Actor is marked dirty */
+	void BroadcastHLODActorMarkedDirty(class ALODActor* InActor) { HLODActorMarkedDirtyEvent.Broadcast(InActor); }
+
+	/** Editor-only event triggered when a HLOD Actor is marked dirty */
+	DECLARE_EVENT(UEngine, FHLODDrawDistanceChangedEvent);
+	FHLODDrawDistanceChangedEvent& OnHLODDrawDistanceChanged() { return HLODDrawDistanceChangedEvent; }
+
+	/** Called by internal engine systems after a HLOD Actor is marked dirty */
+	void BroadcastHLODDrawDistanceChanged() { HLODDrawDistanceChangedEvent.Broadcast(); }
+
 #endif // #if WITH_EDITOR
 
 	/** Event triggered after a server travel failure of any kind has occurred */
@@ -2264,6 +2299,21 @@ private:
 	/** Broadcasts after an actor has been moved, rotated or scaled */
 	FOnActorMovedEvent		OnActorMovedEvent;
 
+	/** Broadcasts after an HLOD actor has been moved between clusters */	
+	FHLODActorMovedEvent HLODActorMovedEvent;
+
+	/** Broadcasts after an HLOD actor's mesh is build*/
+	FHLODMeshBuildEvent HLODMeshBuildEvent;
+	
+	/** Broadcasts after an HLOD actor has added to a cluster */
+	FHLODActorAddedEvent HLODActorAddedEvent;
+
+	/** Broadcasts after an HLOD actor has been marked dirty */
+	FHLODActorMarkedDirtyEvent HLODActorMarkedDirtyEvent;
+
+	/** Broadcasts after a Draw distance value (World settings) is changed */
+	FHLODDrawDistanceChangedEvent HLODDrawDistanceChangedEvent;
+
 #endif // #if WITH_EDITOR
 
 	/** Thread preventing screen saver from kicking. Suspend most of the time. */
diff --git a/Engine/Source/Runtime/Engine/Classes/Engine/LODActor.h b/Engine/Source/Runtime/Engine/Classes/Engine/LODActor.h
--- a/Engine/Source/Runtime/Engine/Classes/Engine/LODActor.h
+++ b/Engine/Source/Runtime/Engine/Classes/Engine/LODActor.h
@@ -13,7 +13,7 @@
 UCLASS(notplaceable, hidecategories=(Object, Collision, Display, Input, Blueprint, Transform))
 class ENGINE_API ALODActor : public AActor
 {
-	GENERATED_UCLASS_BODY()
+GENERATED_UCLASS_BODY()
 
 private_subobject:
 	// disable display of this component
@@ -45,9 +45,54 @@ public:
 
 	/** Sets StaticMesh and IsPreviewActor to true if InStaticMesh equals nullptr */
 	void SetStaticMesh(class UStaticMesh* InStaticMesh);
+
+
+#if WITH_EDITOR
+	/**
+	* Adds InAcor to the SubActors array and set its LODParent to this
+	* @param InActor - Actor to add
+	*/
+	void AddSubActor(AActor* InActor);
+
+	/**
+	* Removes InActor from the SubActors array and sets its LODParent to nullptr
+	* @param InActor - Actor to remove
+	*/
+	void RemoveSubActor(AActor* InActor);
+
+	/**
+	* Returns whether or not this LODActor is dirty
+	* @return const bool
+	*/
+	const bool IsDirty() const { return bDirty; }
+
+	/**
+	* Sets whether or not this LODActor is dirty and should have its LODMesh (re)build
+	* @param bNewState - New dirty state
+	*/
+	void SetIsDirty(const bool bNewState);
+	
+	/**
+	* Determines whether or not this LODActor has valid SubActors (more than one, or one staticmesh actor)
+	* @return const bool
+	*/
+	const bool HasValidSubActors();
+
+	/** Toggles forcing the StaticMeshComponent drawing distance to 0 or LODDrawDistance */
+	void ToggleForceView();
+
+	/** Sets forcing the StaticMeshComponent drawing distance to 0 or LODDrawDistance according to InState*/
+	void SetForcedView(const bool InState);
+
+	/** Sets the state of whether or not this LODActor is hidden from the Editor view, used for forcing a HLOD to show*/
+	void SetHiddenFromEditorView(const bool InState, const int32 ForceLODLevel);
+
+	/** Returns the number of triangles this LODActor's SubActors contain */
+	const uint32 GetNumTriangles();
 	
-	/** Returns whether or not this is a preview actor instance, meaning it has no StaticMesh */
-	const bool IsPreviewActorInstance() { return bIsPreviewActor;  }
+	/** Updates the LODParents for the SubActors (and the drawing distance)*/
+	void UpdateSubActorLODParents();
+#endif // WITH_EDITOR
 	
 protected:
 	// Begin UObject interface.
@@ -61,24 +106,20 @@ protected:
 	virtual void EditorApplyMirror(const FVector& MirrorScale, const FVector& PivotLocation) override;
 #endif // WITH_EDITOR	
 	virtual void PostRegisterAllComponents() override;
-	// End UObject interface.
-		
+	// End UObject interface.		
 protected:
 #if WITH_EDITORONLY_DATA
-	/** Visualization component for rendering the LODActors bounds (cluster bounds)*/
+	/** Whether or not this LODActor is not build or needs rebuilding */
 	UPROPERTY()
-	class UDrawSphereComponent* DrawSphereComponent;
+	bool bDirty;
 #endif // WITH_EDITORONLY_DATA
-
-	/** Whether or not this is a preview actor, will not fire StaticMesh map errors if so*/
-	bool bIsPreviewActor;
-
 public:
-	/** Returns StaticMeshComponent subobject **/
-	class UStaticMeshComponent* GetStaticMeshComponent() const { return StaticMeshComponent; }
-
 #if WITH_EDITORONLY_DATA
-	/** Returns StaticMeshComponent subobject **/
-	class UDrawSphereComponent* GetDrawSphereComponent() const { return DrawSphereComponent; }
+	/** Cached number of triangles contained in the SubActors*/
+	UPROPERTY()
+	uint32 NumTrianglesInSubActors;
 #endif // WITH_EDITORONLY_DATA
-}; 
+
+	/** Returns StaticMeshComponent subobject **/
+	class UStaticMeshComponent* GetStaticMeshComponent() const { return StaticMeshComponent; }
+}; 
\ No newline at end of file
diff --git a/Engine/Source/Runtime/Engine/Classes/GameFramework/WorldSettings.h b/Engine/Source/Runtime/Engine/Classes/GameFramework/WorldSettings.h
--- a/Engine/Source/Runtime/Engine/Classes/GameFramework/WorldSettings.h
+++ b/Engine/Source/Runtime/Engine/Classes/GameFramework/WorldSettings.h
@@ -240,8 +240,8 @@ struct ENGINE_API FHierarchicalSimplification
 
 	FHierarchicalSimplification()
 		: bSimplifyMesh(false)
-		, DrawDistance(3000)
-		, DesiredBoundRadius(3000)
+		, DrawDistance(2000)
+		, DesiredBoundRadius(2000)
 		, DesiredFillingPercentage(50)
 		, MinNumberOfActorsToBuild(2)
 	{
diff --git a/Engine/Source/Runtime/Engine/Private/LODActor.cpp b/Engine/Source/Runtime/Engine/Private/LODActor.cpp
--- a/Engine/Source/Runtime/Engine/Private/LODActor.cpp
+++ b/Engine/Source/Runtime/Engine/Private/LODActor.cpp
@@ -10,9 +10,7 @@
 #include "MessageLog.h"
 #include "UObjectToken.h"
 
-#if WITH_EDITOR	
-#include "Components/DrawSphereComponent.h"
-#endif // WITH_EDITOR
+#include "StaticMeshResources.h"
 
 #define LOCTEXT_NAMESPACE "LODActor"
 
@@ -21,7 +19,16 @@ ALODActor::ALODActor(const FObjectInitializer& ObjectInitializer)
 	, LODDrawDistance(5000)	
 {
 	bCanBeDamaged = false;
-	bIsPreviewActor = false;
+
+#if WITH_EDITORONLY_DATA
+	
+	bListedInSceneOutliner = false;
+
+	// Always dirty when created
+	bDirty = true;
+
+	NumTrianglesInSubActors = 0;
+#endif // WITH_EDITORONLY_DATA
 
 	StaticMeshComponent = CreateDefaultSubobject<UStaticMeshComponent>(TEXT("StaticMeshComponent0"));
 	StaticMeshComponent->SetCollisionProfileName(UCollisionProfile::NoCollision_ProfileName);
@@ -31,15 +38,7 @@ ALODActor::ALODActor(const FObjectInitializer& ObjectInitializer)
 	StaticMeshComponent->bCastStaticShadow = false;
 	StaticMeshComponent->CastShadow = false;
 
-	RootComponent = StaticMeshComponent;
-	
-#if WITH_EDITORONLY_DATA	
-	DrawSphereComponent = CreateEditorOnlyDefaultSubobject<UDrawSphereComponent>(TEXT("VisualizeComponent0"));
-	if (DrawSphereComponent)
-	{
-		DrawSphereComponent->SetSphereRadius(0.0f);
-	}
-#endif // WITH_EDITORONLY_DATA
+	RootComponent = StaticMeshComponent;	
 }
 
 FString ALODActor::GetDetailedInfoInternal() const
@@ -51,24 +50,18 @@ void ALODActor::PostRegisterAllComponents()
 {
 	Super::PostRegisterAllComponents();
 #if WITH_EDITOR
-	if (!bIsPreviewActor && StaticMeshComponent->SceneProxy)
+	if (StaticMeshComponent->SceneProxy)
 	{
 		ensure (LODLevel >= 1);
 		(StaticMeshComponent->SceneProxy)->SetHierarchicalLOD_GameThread(LODLevel);
 	}
+
+	UpdateSubActorLODParents();
+
 #endif
 }
 
-void ALODActor::SetStaticMesh(class UStaticMesh* InStaticMesh)
-{
-	if (StaticMeshComponent)
-	{
-		StaticMeshComponent->StaticMesh = InStaticMesh;
-	}
 
-	// If given NULL for static mesh then this is a preview actor (only rendering the bounds)
-	bIsPreviewActor = (InStaticMesh == nullptr);
-}
 
 #if WITH_EDITOR
 
@@ -97,6 +90,7 @@ void ALODActor::PostEditChangeProperty(FPropertyChangedEvent& PropertyChangedEve
 				}
 			}
 		}
+		SetIsDirty(true);
 	}
 	Super::PostEditChangeProperty(PropertyChangedEvent);
 }
@@ -112,30 +106,26 @@ void ALODActor::CheckForErrors()
 {
 	FMessageLog MapCheck("MapCheck");
 
-	// Only check when this is not a preview actor and actually has a static mesh
-	if (!bIsPreviewActor)
+	// Only check when this is not a preview actor and actually has a static mesh	
+	Super::CheckForErrors();
+	if (!StaticMeshComponent)
 	{
-		Super::CheckForErrors();
-		if (!StaticMeshComponent)
-		{
-			MapCheck.Warning()
-				->AddToken(FUObjectToken::Create(this))
-				->AddToken(FTextToken::Create(LOCTEXT("MapCheck_Message_StaticMeshComponent", "Static mesh actor has NULL StaticMeshComponent property - please delete")))
-				->AddToken(FMapErrorToken::Create(FMapErrors::StaticMeshComponent));
-		}
+		MapCheck.Warning()
+			->AddToken(FUObjectToken::Create(this))
+			->AddToken(FTextToken::Create(LOCTEXT("MapCheck_Message_StaticMeshComponent", "Static mesh actor has NULL StaticMeshComponent property - please delete")))
+			->AddToken(FMapErrorToken::Create(FMapErrors::StaticMeshComponent));
+	}
 
-		if (StaticMeshComponent && StaticMeshComponent->StaticMesh == NULL)
-		{
-			FFormatNamedArguments Arguments;
-			Arguments.Add(TEXT("ActorName"), FText::FromString(GetName()));
-			FMessageLog("MapCheck").Error()
-				->AddToken(FUObjectToken::Create(this))
-				->AddToken(FTextToken::Create(FText::Format(LOCTEXT("MapCheck_Message_InvalidLODActorMissingMesh", "{ActorName} : Static mesh is missing for the built LODActor.  Did you remove the asset? Please delete it and build LOD again. "), Arguments)))
-				->AddToken(FMapErrorToken::Create(FMapErrors::LODActorMissingStaticMesh));
-		}
+	if (StaticMeshComponent && StaticMeshComponent->StaticMesh == NULL)
+	{
+		FFormatNamedArguments Arguments;
+		Arguments.Add(TEXT("ActorName"), FText::FromString(GetName()));
+		FMessageLog("MapCheck").Error()
+			->AddToken(FUObjectToken::Create(this))
+			->AddToken(FTextToken::Create(FText::Format(LOCTEXT("MapCheck_Message_InvalidLODActorMissingMesh", "{ActorName} : Static mesh is missing for the built LODActor.  Did you remove the asset? Please delete it and build LOD again. "), Arguments)))
+			->AddToken(FMapErrorToken::Create(FMapErrors::LODActorMissingStaticMesh));
 	}
 	
-
 	if (SubActors.Num() == 0)
 	{
 		FFormatNamedArguments Arguments;
@@ -179,12 +169,204 @@ void ALODActor::EditorApplyMirror(const FVector& MirrorScale, const FVector& Piv
 {
 }
 
+void ALODActor::AddSubActor(AActor* InActor)
+{
+	SubActors.Add(InActor);
+	InActor->SetLODParent(StaticMeshComponent, LODDrawDistance);
+	SetIsDirty(true);
+
+	// Adding number of triangles
+	if (!InActor->IsA<ALODActor>())
+	{
+		TArray<UStaticMeshComponent*> StaticMeshComponents;
+		InActor->GetComponents<UStaticMeshComponent>(StaticMeshComponents);
+		for (UStaticMeshComponent* Component : StaticMeshComponents)
+		{
+			if (Component && Component->StaticMesh && Component->StaticMesh->RenderData)
+			{
+				NumTrianglesInSubActors += Component->StaticMesh->RenderData->LODResources[0].GetNumTriangles();
+			}
+
+			Component->MarkRenderStateDirty();
+		}
+	}
+	else
+	{
+		ALODActor* LODActor = Cast<ALODActor>(InActor);
+		NumTrianglesInSubActors += LODActor->GetNumTriangles();
+	}
+
+	StaticMeshComponent->MarkRenderStateDirty();
+	
+}
+
+void ALODActor::RemoveSubActor(AActor* InActor)
+{
+	SubActors.Remove(InActor);
+	InActor->SetLODParent(nullptr, 0);
+	SetIsDirty(true);
+	
+	// Deducting number of triangles
+	if (!InActor->IsA<ALODActor>())
+	{
+		TArray<UStaticMeshComponent*> StaticMeshComponents;
+		InActor->GetComponents<UStaticMeshComponent>(StaticMeshComponents);
+		for (UStaticMeshComponent* Component : StaticMeshComponents)
+		{
+			if (Component && Component->StaticMesh && Component->StaticMesh->RenderData)
+			{
+				NumTrianglesInSubActors -= Component->StaticMesh->RenderData->LODResources[0].GetNumTriangles();
+			}
+
+			Component->MarkRenderStateDirty();
+		}
+	}
+	else
+	{
+		ALODActor* LODActor = Cast<ALODActor>(InActor);
+		NumTrianglesInSubActors -= LODActor->GetNumTriangles();
+	}
+
+	StaticMeshComponent->MarkRenderStateDirty();
+
+	// In case the user removes an actor while the HLOD system is force viewing one LOD level
+	InActor->SetIsTemporarilyHiddenInEditor(false);
+}
+
+void ALODActor::SetIsDirty(const bool bNewState)
+{
+	bDirty = bNewState;
+
+	// Set parent LODActor dirty as well if bNewState = true
+	if (IsDirty())
+	{
+		// If this LODActor is a SubActor at a higher LOD level mark parent dirty as well
+		UPrimitiveComponent* ParentComponent = StaticMeshComponent->GetLODParentPrimitive();
+		if (ParentComponent)
+		{
+			ALODActor* LODParentActor = Cast<ALODActor>(ParentComponent->GetOwner());
+			if (LODParentActor)
+			{
+				LODParentActor->SetIsDirty(true);
+			}
+		}
+
+		// Broadcast actor marked dirty event
+		if (GEngine)
+		{
+			GEngine->BroadcastHLODActorMarkedDirty(this);
+		}		
+	}	
+	else
+	{
+		// Update SubActor's LOD parent component
+		for (auto& SubActor : SubActors)
+		{
+			SubActor->SetLODParent(StaticMeshComponent, LODDrawDistance);
+		}
+	}
+}
+
+const bool ALODActor::HasValidSubActors()
+{
+	if (SubActors.Num() > 1)
+	{
+		// More than one sub actor
+		return true;
+	}
+	else if (SubActors.Num() == 0)
+	{
+		// No sub actors
+		return false;
+	}	
+	else
+	{
+		if (SubActors[0]->IsA<ALODActor>())
+		{
+			// LODActor as only sub actor (which doesn't make sense for the system)
+			return false;
+		}
+		else
+		{
+			// StaticMeshActor as only sub actor
+			return true;
+		}
+	}
+	
+}
+
+void ALODActor::ToggleForceView()
+{
+	// Toggle the forced viewing of this LODActor, set drawing distance to 0.0f or LODDrawDistance
+	StaticMeshComponent->MinDrawDistance = (StaticMeshComponent->MinDrawDistance == 0.0f) ? LODDrawDistance : 0.0f;
+	StaticMeshComponent->MarkRenderStateDirty();
+}
+
+void ALODActor::SetForcedView(const bool InState)
+{
+	// Set forced viewing state of this LODActor, set drawing distance to 0.0f or LODDrawDistance
+	StaticMeshComponent->MinDrawDistance = (InState) ? 0.0f : LODDrawDistance;
+	StaticMeshComponent->MarkRenderStateDirty();
+}
+
+void ALODActor::SetHiddenFromEditorView(const bool InState, const int32 ForceLODLevel )
+{
+	// If we are also subactor for a higher LOD level or this actor belongs to a higher HLOD level than is being forced hide the actor
+	if (GetStaticMeshComponent()->GetLODParentPrimitive() || LODLevel > ForceLODLevel )
+	{
+		SetIsTemporarilyHiddenInEditor(InState);		
+	}
+
+	for (auto Actor : SubActors)
+	{
+		// If this actor belongs to a lower HLOD level that is being forced hide the sub-actors
+		if (LODLevel < ForceLODLevel)
+		{
+			Actor->SetIsTemporarilyHiddenInEditor(InState);
+		}
+
+		// Toggle/set the LOD parent to nullptr or this
+		Actor->SetLODParent((InState) ? nullptr : StaticMeshComponent, (InState) ? 0.0f : LODDrawDistance);
+	}
+
+	StaticMeshComponent->MarkRenderStateDirty();
+}
+
+const uint32 ALODActor::GetNumTriangles()
+{
+	return NumTrianglesInSubActors;
+}
+
+
+void ALODActor::SetStaticMesh(class UStaticMesh* InStaticMesh)
+{
+	if (StaticMeshComponent)
+	{
+		StaticMeshComponent->StaticMesh = InStaticMesh;
+		SetIsDirty(false);
+	}
+}
+
+void ALODActor::UpdateSubActorLODParents()
+{
+	for (auto& Actor : SubActors)
+	{
+		Actor->SetLODParent(StaticMeshComponent, LODDrawDistance);
+	}
+}
+
 #endif // WITH_EDITOR
 
 FBox ALODActor::GetComponentsBoundingBox(bool bNonColliding) const 
 {
 	FBox BoundBox = Super::GetComponentsBoundingBox(bNonColliding);
 
+	// If BoundBox ends up to nothing create a new invalid one
+	if (BoundBox.GetVolume() == 0.0f)
+	{
+		BoundBox = FBox(0);
+	}
+
 	if (bNonColliding)
 	{
 		if (StaticMeshComponent && StaticMeshComponent->StaticMesh)
@@ -193,10 +375,19 @@ FBox ALODActor::GetComponentsBoundingBox(bool bNonColliding) const
 			FBox StaticBoundBox(BoundBox.GetCenter()-StaticBound.BoxExtent, BoundBox.GetCenter()+StaticBound.BoxExtent);
 			BoundBox += StaticBoundBox;
 		}
+		else
+		{
+			for (auto Actor : SubActors)
+			{
+				if (Actor)
+				{
+					BoundBox += Actor->GetComponentsBoundingBox(bNonColliding);
+				}				
+			}
+		}
 	}
 
 	return BoundBox;	
 }
 
 #undef LOCTEXT_NAMESPACE
-
diff --git a/Engine/Source/Runtime/Engine/Private/WorldSettings.cpp b/Engine/Source/Runtime/Engine/Private/WorldSettings.cpp
--- a/Engine/Source/Runtime/Engine/Private/WorldSettings.cpp
+++ b/Engine/Source/Runtime/Engine/Private/WorldSettings.cpp
@@ -16,6 +16,10 @@
 
 #define LOCTEXT_NAMESPACE "ErrorChecking"
 
+#if WITH_EDITOR
+#include "Editor.h"
+#endif 
+
 AWorldSettings::AWorldSettings(const FObjectInitializer& ObjectInitializer)
 	: Super(ObjectInitializer.DoNotCreateDefaultSubobject(TEXT("Sprite")))
 {
@@ -294,6 +298,11 @@ void AWorldSettings::PostEditChangeProperty(FPropertyChangedEvent& PropertyChang
 		{
 			GEngine->DeferredCommands.AddUnique(TEXT("UpdateLandscapeSetup"));
 		}
+
+		if (PropertyThatChanged->GetName() == TEXT("DrawDistance"))
+		{
+			GEditor->BroadcastHLODDrawDistanceChanged();
+		}
 	}
 
 	Super::PostEditChangeProperty(PropertyChangedEvent);
