HLOD outliner build fix

diff --git a/Engine/Source/Editor/HierarchicalLODOutliner/Private/HLODSelectionActor.h b/Engine/Source/Editor/HierarchicalLODOutliner/Private/HLODSelectionActor.h
--- a/Engine/Source/Editor/HierarchicalLODOutliner/Private/HLODSelectionActor.h
+++ b/Engine/Source/Editor/HierarchicalLODOutliner/Private/HLODSelectionActor.h
@@ -2,6 +2,7 @@
 
 #pragma once
 #include "HLODSelectionActor.generated.h"
+#include "GameFramework\Actor.h"
 
 UCLASS(notplaceable, hidecategories = (Object, Collision, Display, Input, Blueprint, Transform))
 class AHLODSelectionActor : public AActor
