Outliner missing includes and linking error fix

diff --git a/Engine/Source/Editor/HierarchicalLODOutliner/Private/HLODOutliner.h b/Engine/Source/Editor/HierarchicalLODOutliner/Private/HLODOutliner.h
--- a/Engine/Source/Editor/HierarchicalLODOutliner/Private/HLODOutliner.h
+++ b/Engine/Source/Editor/HierarchicalLODOutliner/Private/HLODOutliner.h
@@ -7,6 +7,7 @@
 #include "SListView.h"
 #include "STableViewBase.h"
 #include "TreeItemID.h"
+#include "SlateBasics.h"
 
 class AHLODSelectionActor;
 class UWorld;
diff --git a/Engine/Source/Editor/HierarchicalLODOutliner/Private/HLODSelectionActor.h b/Engine/Source/Editor/HierarchicalLODOutliner/Private/HLODSelectionActor.h
--- a/Engine/Source/Editor/HierarchicalLODOutliner/Private/HLODSelectionActor.h
+++ b/Engine/Source/Editor/HierarchicalLODOutliner/Private/HLODSelectionActor.h
@@ -1,7 +1,7 @@
 // Copyright 1998-2015 Epic Games, Inc. All Rights Reserved.
 
 #pragma once
-#include "GameFramework\Actor.h"
+#include "GameFramework/Actor.h"
 #include "HLODSelectionActor.generated.h"
 
 UCLASS(notplaceable, hidecategories = (Object, Collision, Display, Input, Blueprint, Transform))
