Another missed include, didn't come up in non-unity builds

diff --git a/Engine/Source/Editor/HierarchicalLODOutliner/Public/HierarchicalLODOutlinerModule.h b/Engine/Source/Editor/HierarchicalLODOutliner/Public/HierarchicalLODOutlinerModule.h
--- a/Engine/Source/Editor/HierarchicalLODOutliner/Public/HierarchicalLODOutlinerModule.h
+++ b/Engine/Source/Editor/HierarchicalLODOutliner/Public/HierarchicalLODOutlinerModule.h
@@ -1,6 +1,7 @@
 // Copyright 1998-2015 Epic Games, Inc. All Rights Reserved.
 
 #include "ModuleInterface.h"
+#include "SlateBasics.h"
 
 /**
 * The module holding all of the UI related pieces for HLOD Outliner
