HLOD - Main
- Fixed spelling error (build/built)
- Changed colour of not yet built LODActors in Outliner
- More includes hoping to fix abstract build errors

diff --git a/Engine/Source/Editor/HierarchicalLODOutliner/Private/HLODOutliner.h b/Engine/Source/Editor/HierarchicalLODOutliner/Private/HLODOutliner.h
--- a/Engine/Source/Editor/HierarchicalLODOutliner/Private/HLODOutliner.h
+++ b/Engine/Source/Editor/HierarchicalLODOutliner/Private/HLODOutliner.h
@@ -3,6 +3,7 @@
 #pragma once
 
 #include "SCompoundWidget.h"
+#include "UnrealEd.h"
 #include "STreeView.h"
 #include "SListView.h"
 #include "STableViewBase.h"
diff --git a/Engine/Source/Editor/HierarchicalLODOutliner/Private/HLODSelectionActor.h b/Engine/Source/Editor/HierarchicalLODOutliner/Private/HLODSelectionActor.h
--- a/Engine/Source/Editor/HierarchicalLODOutliner/Private/HLODSelectionActor.h
+++ b/Engine/Source/Editor/HierarchicalLODOutliner/Private/HLODSelectionActor.h
@@ -1,6 +1,7 @@
 // Copyright 1998-2015 Epic Games, Inc. All Rights Reserved.
 
 #pragma once
+#include "Engine.h"
 #include "GameFramework/Actor.h"
 #include "HLODSelectionActor.generated.h"
 
diff --git a/Engine/Source/Editor/HierarchicalLODOutliner/Private/HierarchicalLODOutlinerPrivatePCH.h b/Engine/Source/Editor/HierarchicalLODOutliner/Private/HierarchicalLODOutlinerPrivatePCH.h
--- a/Engine/Source/Editor/HierarchicalLODOutliner/Private/HierarchicalLODOutlinerPrivatePCH.h
+++ b/Engine/Source/Editor/HierarchicalLODOutliner/Private/HierarchicalLODOutlinerPrivatePCH.h
@@ -3,3 +3,4 @@
 #pragma once
 
 #include "Core.h"
+#include "UnrealEd.h"
diff --git a/Engine/Source/Editor/HierarchicalLODOutliner/Private/LODActorItem.cpp b/Engine/Source/Editor/HierarchicalLODOutliner/Private/LODActorItem.cpp
--- a/Engine/Source/Editor/HierarchicalLODOutliner/Private/LODActorItem.cpp
+++ b/Engine/Source/Editor/HierarchicalLODOutliner/Private/LODActorItem.cpp
@@ -56,7 +56,7 @@ FString HLODOutliner::FLODActorItem::GetDisplayString() const
 {
 	if (ALODActor* ActorPtr = LODActor.Get())
 	{
-		return ActorPtr->GetName() + ((ActorPtr->IsDirty()) ? FString(" (Not build)") : FString());
+		return ActorPtr->GetName() + ((ActorPtr->IsDirty()) ? FString(" (Not built)") : FString());
 	}
 
 	return FString();
@@ -72,7 +72,7 @@ FSlateColor HLODOutliner::FLODActorItem::GetTint() const
 	ALODActor* LODActorPtr = LODActor.Get();
 	if (LODActorPtr)
 	{
-		return LODActorPtr->IsDirty() ? FLinearColor(1.0f, 0.6f, 0.6f) : FLinearColor(1.0f, 1.0f, 1.0f);
+		return LODActorPtr->IsDirty() ? FSlateColor::UseSubduedForeground() : FLinearColor(1.0f, 1.0f, 1.0f);
 	}
 
 	return FLinearColor(1.0f, 1.0f, 1.0f);
diff --git a/Engine/Source/Editor/HierarchicalLODOutliner/Public/HierarchicalLODOutlinerModule.h b/Engine/Source/Editor/HierarchicalLODOutliner/Public/HierarchicalLODOutlinerModule.h
--- a/Engine/Source/Editor/HierarchicalLODOutliner/Public/HierarchicalLODOutlinerModule.h
+++ b/Engine/Source/Editor/HierarchicalLODOutliner/Public/HierarchicalLODOutlinerModule.h
@@ -1,7 +1,7 @@
 // Copyright 1998-2015 Epic Games, Inc. All Rights Reserved.
 
 #include "ModuleInterface.h"
-#include "SlateBasics.h"
+#include "UnrealEd.h"
 
 /**
 * The module holding all of the UI related pieces for HLOD Outliner
