CIS Fix. Missing Undefine for LOCTEXT_NAMESPACE

diff --git a/Engine/Source/Editor/HierarchicalLODOutliner/Private/LODActorItem.cpp b/Engine/Source/Editor/HierarchicalLODOutliner/Private/LODActorItem.cpp
--- a/Engine/Source/Editor/HierarchicalLODOutliner/Private/LODActorItem.cpp
+++ b/Engine/Source/Editor/HierarchicalLODOutliner/Private/LODActorItem.cpp
@@ -319,4 +319,6 @@ void HLODOutliner::FLODActorDropTarget::MergeCluster(ALODActor* ToMergeActor)
 		AActor* SubActor = ToMergeActor->SubActors.Last();
 		MoveToCluster(SubActor, LODActor.Get());
 	}
-}
\ No newline at end of file
+}
+
+#undef LOCTEXT_NAMESPACE 
\ No newline at end of file
