HLOD build fixes

diff --git a/Engine/Source/Editor/HierarchicalLODOutliner/Private/HLODOutliner.cpp b/Engine/Source/Editor/HierarchicalLODOutliner/Private/HLODOutliner.cpp
--- a/Engine/Source/Editor/HierarchicalLODOutliner/Private/HLODOutliner.cpp
+++ b/Engine/Source/Editor/HierarchicalLODOutliner/Private/HLODOutliner.cpp
@@ -21,6 +21,7 @@
 #include "StaticMeshActorItem.h"
 #include "HLODTreeWidgetitem.h"
 #include "HLODSelectionActor.h"
+#include "TreeItemID.h"
 
 #include "HierarchicalLODUtils.h"
 
diff --git a/Engine/Source/Editor/HierarchicalLODOutliner/Private/HLODOutlinerDragDrop.cpp b/Engine/Source/Editor/HierarchicalLODOutliner/Private/HLODOutlinerDragDrop.cpp
--- a/Engine/Source/Editor/HierarchicalLODOutliner/Private/HLODOutlinerDragDrop.cpp
+++ b/Engine/Source/Editor/HierarchicalLODOutliner/Private/HLODOutlinerDragDrop.cpp
@@ -6,6 +6,7 @@
 #include "HLODOutlinerDragDrop.h"
 #include "HierarchicalLODUtils.h"
 #include "ITreeItem.h"
+#include "TreeItemID.h"
 
 HLODOutliner::FDragDropPayload::FDragDropPayload()
 {
diff --git a/Engine/Source/Editor/HierarchicalLODOutliner/Private/ITreeItem.h b/Engine/Source/Editor/HierarchicalLODOutliner/Private/ITreeItem.h
--- a/Engine/Source/Editor/HierarchicalLODOutliner/Private/ITreeItem.h
+++ b/Engine/Source/Editor/HierarchicalLODOutliner/Private/ITreeItem.h
@@ -3,7 +3,6 @@
 
 #include "Engine/LodActor.h"
 #include "HLODOutlinerDragDrop.h"
-#include "TreeItemID.h"
 
 class FMenuBuilder;
 class UWorld;
@@ -11,6 +10,7 @@ class UWorld;
 namespace HLODOutliner
 {
 	class SHLODOutliner;
+	struct FTreeItemID;
 
 	/** Interface used for validating movement (ie, drag/dropping) operations */
 	struct IDropTarget
diff --git a/Engine/Source/Editor/HierarchicalLODOutliner/Private/LODActorItem.cpp b/Engine/Source/Editor/HierarchicalLODOutliner/Private/LODActorItem.cpp
--- a/Engine/Source/Editor/HierarchicalLODOutliner/Private/LODActorItem.cpp
+++ b/Engine/Source/Editor/HierarchicalLODOutliner/Private/LODActorItem.cpp
@@ -8,7 +8,7 @@
 #include "HierarchicalLODUtils.h"
 #include "ScopedTransaction.h"
 #include "SlateBasics.h"
-
+#include "TreeItemID.h"
 
 #define LOCTEXT_NAMESPACE "LODActorItem"
 
@@ -321,4 +321,4 @@ void HLODOutliner::FLODActorDropTarget::MergeCluster(ALODActor* ToMergeActor)
 	}
 }
 
-#undef LOCTEXT_NAMESPACE 
\ No newline at end of file
+#undef LOCTEXT_NAMESPACE
diff --git a/Engine/Source/Editor/HierarchicalLODOutliner/Private/LODLevelItem.cpp b/Engine/Source/Editor/HierarchicalLODOutliner/Private/LODLevelItem.cpp
--- a/Engine/Source/Editor/HierarchicalLODOutliner/Private/LODLevelItem.cpp
+++ b/Engine/Source/Editor/HierarchicalLODOutliner/Private/LODLevelItem.cpp
@@ -7,7 +7,7 @@
 #include "HierarchicalLODUtils.h"
 #include "ScopedTransaction.h"
 #include "SlateBasics.h"
-
+#include "TreeItemID.h"
 
 #define LOCTEXT_NAMESPACE "LODLevelItem"
 
@@ -127,5 +127,3 @@ void HLODOutliner::FLODLevelDropTarget::CreateNewCluster(FDragDropPayload &Dragg
 }
 
 #undef LOCTEXT_NAMESPACE
-
-
diff --git a/Engine/Source/Editor/HierarchicalLODOutliner/Private/StaticMeshActorItem.cpp b/Engine/Source/Editor/HierarchicalLODOutliner/Private/StaticMeshActorItem.cpp
--- a/Engine/Source/Editor/HierarchicalLODOutliner/Private/StaticMeshActorItem.cpp
+++ b/Engine/Source/Editor/HierarchicalLODOutliner/Private/StaticMeshActorItem.cpp
@@ -5,6 +5,7 @@
 #include "HLODOutliner.h"
 #include "LevelEditor.h"
 #include "SlateBasics.h"
+#include "TreeItemID.h"
 
 #define LOCTEXT_NAMESPACE "StaticMeshActorItem"
 
diff --git a/Engine/Source/Editor/HierarchicalLODOutliner/Private/TreeItemID.h b/Engine/Source/Editor/HierarchicalLODOutliner/Private/TreeItemID.h
--- a/Engine/Source/Editor/HierarchicalLODOutliner/Private/TreeItemID.h
+++ b/Engine/Source/Editor/HierarchicalLODOutliner/Private/TreeItemID.h
@@ -8,6 +8,8 @@ namespace HLODOutliner
 	struct FTreeItemID
 	{
 	public:
+		FTreeItemID() : CachedHash(0) {}
+
 		FTreeItemID(const UObject* InObject)
 		{
 			ObjectKey = FObjectKey(InObject);
@@ -20,11 +22,19 @@ namespace HLODOutliner
 			*this = Other;
 		}
 
+		FTreeItemID& operator=(const FTreeItemID& Other)
+		{
+			ObjectKey = Other.ObjectKey;
+			CachedHash = Other.CachedHash;
+			return *this;
+		}
+
 		/** Move construction / assignment */
 		FTreeItemID(FTreeItemID&& Other)
 		{
 			*this = MoveTemp(Other);
 		}
+
 		FTreeItemID& operator=(FTreeItemID&& Other)
 		{
 			FMemory::Memswap(this, &Other, sizeof(FTreeItemID));
