Fix compilation on a case-sensitive fs.

diff --git a/Engine/Source/Editor/HierarchicalLODOutliner/Private/HLODOutliner.cpp b/Engine/Source/Editor/HierarchicalLODOutliner/Private/HLODOutliner.cpp
--- a/Engine/Source/Editor/HierarchicalLODOutliner/Private/HLODOutliner.cpp
+++ b/Engine/Source/Editor/HierarchicalLODOutliner/Private/HLODOutliner.cpp
@@ -4,7 +4,7 @@
 #include "HLODOutliner.h"
 
 #include "Engine.h"
-#include "Engine/LodActor.h"
+#include "Engine/LODActor.h"
 #include "Engine/World.h"
 #include "HierarchicalLOD.h"
 
@@ -19,7 +19,7 @@
 #include "LODActorItem.h"
 #include "LODLevelItem.h"
 #include "StaticMeshActorItem.h"
-#include "HLODTreeWidgetitem.h"
+#include "HLODTreeWidgetItem.h"
 #include "HLODSelectionActor.h"
 #include "TreeItemID.h"
 
diff --git a/Engine/Source/Editor/HierarchicalLODOutliner/Private/ITreeItem.h b/Engine/Source/Editor/HierarchicalLODOutliner/Private/ITreeItem.h
--- a/Engine/Source/Editor/HierarchicalLODOutliner/Private/ITreeItem.h
+++ b/Engine/Source/Editor/HierarchicalLODOutliner/Private/ITreeItem.h
@@ -1,7 +1,7 @@
 // Copyright 1998-2015 Epic Games, Inc. All Rights Reserved.
 #pragma once
 
-#include "Engine/LodActor.h"
+#include "Engine/LODActor.h"
 #include "HLODOutlinerDragDrop.h"
 
 class FMenuBuilder;
diff --git a/Engine/Source/Editor/HierarchicalLODOutliner/Private/LODActorItem.cpp b/Engine/Source/Editor/HierarchicalLODOutliner/Private/LODActorItem.cpp
--- a/Engine/Source/Editor/HierarchicalLODOutliner/Private/LODActorItem.cpp
+++ b/Engine/Source/Editor/HierarchicalLODOutliner/Private/LODActorItem.cpp
@@ -3,7 +3,7 @@
 #include "HierarchicalLODOutlinerPrivatePCH.h"
 #include "LevelEditor.h"
 #include "LODActorItem.h"
-#include "Engine/LodActor.h"
+#include "Engine/LODActor.h"
 #include "HLODOutliner.h"
 #include "HierarchicalLODUtils.h"
 #include "ScopedTransaction.h"
diff --git a/Engine/Source/Editor/UnrealEd/Public/HierarchicalLODUtils.h b/Engine/Source/Editor/UnrealEd/Public/HierarchicalLODUtils.h
--- a/Engine/Source/Editor/UnrealEd/Public/HierarchicalLODUtils.h
+++ b/Engine/Source/Editor/UnrealEd/Public/HierarchicalLODUtils.h
@@ -6,7 +6,7 @@
 #include "Engine/StaticMesh.h"
 #include "MeshUtilities.h"
 
-#include "Asyncwork.h"
+#include "AsyncWork.h"
 
 #if WITH_EDITOR
 #include "AssetData.h"
