HLOD Non-unity build errors fix

diff --git a/Engine/Source/Editor/HierarchicalLODOutliner/Private/LODActorItem.h b/Engine/Source/Editor/HierarchicalLODOutliner/Private/LODActorItem.h
--- a/Engine/Source/Editor/HierarchicalLODOutliner/Private/LODActorItem.h
+++ b/Engine/Source/Editor/HierarchicalLODOutliner/Private/LODActorItem.h
@@ -2,6 +2,7 @@
 #pragma once
 
 #include "ITreeItem.h"
+#include "TreeItemID.h"
 
 class ALODActor;
 class AActor;
diff --git a/Engine/Source/Editor/HierarchicalLODOutliner/Private/LODLevelItem.h b/Engine/Source/Editor/HierarchicalLODOutliner/Private/LODLevelItem.h
--- a/Engine/Source/Editor/HierarchicalLODOutliner/Private/LODLevelItem.h
+++ b/Engine/Source/Editor/HierarchicalLODOutliner/Private/LODLevelItem.h
@@ -2,6 +2,7 @@
 #pragma once
 
 #include "ITreeItem.h"
+#include "TreeItemID.h"
 
 namespace HLODOutliner
 {
diff --git a/Engine/Source/Editor/HierarchicalLODOutliner/Private/StaticMeshActorItem.cpp b/Engine/Source/Editor/HierarchicalLODOutliner/Private/StaticMeshActorItem.cpp
--- a/Engine/Source/Editor/HierarchicalLODOutliner/Private/StaticMeshActorItem.cpp
+++ b/Engine/Source/Editor/HierarchicalLODOutliner/Private/StaticMeshActorItem.cpp
@@ -6,6 +6,7 @@
 #include "LevelEditor.h"
 #include "SlateBasics.h"
 #include "TreeItemID.h"
+#include "GameFramework/Actor.h"
 
 #define LOCTEXT_NAMESPACE "StaticMeshActorItem"
 
diff --git a/Engine/Source/Editor/HierarchicalLODOutliner/Private/StaticMeshActorItem.h b/Engine/Source/Editor/HierarchicalLODOutliner/Private/StaticMeshActorItem.h
--- a/Engine/Source/Editor/HierarchicalLODOutliner/Private/StaticMeshActorItem.h
+++ b/Engine/Source/Editor/HierarchicalLODOutliner/Private/StaticMeshActorItem.h
@@ -2,7 +2,9 @@
 #pragma once
 
 #include "ITreeItem.h"
+#include "TreeItemID.h"
 
+class AActor;
 
 namespace HLODOutliner
 {
