Fix Mac unity build compilation

#pragma once had been left in a cpp file

diff --git a/Engine/Source/Editor/HierarchicalLODOutliner/Private/HLODOutlinerDragDrop.cpp b/Engine/Source/Editor/HierarchicalLODOutliner/Private/HLODOutlinerDragDrop.cpp
--- a/Engine/Source/Editor/HierarchicalLODOutliner/Private/HLODOutlinerDragDrop.cpp
+++ b/Engine/Source/Editor/HierarchicalLODOutliner/Private/HLODOutlinerDragDrop.cpp
@@ -1,7 +1,5 @@
 // Copyright 1998-2015 Epic Games, Inc. All Rights Reserved.
 
-#pragma once
-
 #include "HierarchicalLODOutlinerPrivatePCH.h"
 #include "HLODOutlinerDragDrop.h"
 #include "HierarchicalLODUtils.h"
