HLOD outliner changes
- Fixed crash when having a selection actor and rebuilding meshes
- Fixed crash when closing HLOD outliner
- Added default values for each newly added HLOD level ( 2.5x Radius/Distance 0.75x FillingRate)
- Fixed ability to create new HLOD cluster node with two child clusters within the same HLOD node (drag and drop action)

diff --git a/Engine/Source/Editor/HierarchicalLODOutliner/Private/HLODOutliner.cpp b/Engine/Source/Editor/HierarchicalLODOutliner/Private/HLODOutliner.cpp
--- a/Engine/Source/Editor/HierarchicalLODOutliner/Private/HLODOutliner.cpp
+++ b/Engine/Source/Editor/HierarchicalLODOutliner/Private/HLODOutliner.cpp
@@ -7,6 +7,7 @@
 #include "Engine/LODActor.h"
 #include "Engine/World.h"
 #include "HierarchicalLOD.h"
+#include "HierarchicalLODVolume.h"
 
 #include "Editor.h"
 #include "EditorModeManager.h"
@@ -14,6 +15,7 @@
 #include "IDetailsView.h"
 #include "SListView.h"
 #include "ScopedTransaction.h"
+#include "BSPOps.h"
 
 #include "ITreeItem.h"
 #include "LODActorItem.h"
@@ -25,6 +27,7 @@
 
 #include "HierarchicalLODUtils.h"
 
+
 #define LOCTEXT_NAMESPACE "HLODOutliner"
 
 namespace HLODOutliner
@@ -80,15 +83,21 @@ namespace HLODOutliner
 			];
 
 		MainContentPanel->AddSlot()
-			.FillHeight(0.6f)
+			.FillHeight(1.0f)
 			[
-				CreateTreeviewWidget()
-			];
+				SNew(SSplitter)
+				.Orientation(Orient_Vertical)
 
-		MainContentPanel->AddSlot()
-			.FillHeight(0.4f)
-			[
-				SettingsView.ToSharedRef()
+				+ SSplitter::Slot()
+				.Value(0.5)
+				[
+					CreateTreeviewWidget()
+				]
+				+ SSplitter::Slot()
+				.Value(0.5)
+				[
+					SettingsView.ToSharedRef()
+				]		
 			];
 
 		RegisterDelegates();
@@ -103,6 +112,7 @@ namespace HLODOutliner
 			+ SVerticalBox::Slot()
 			.Padding(FMargin(0.0f, 5.0f))
 			[
+
 				SNew(SHorizontalBox)
 				+ SHorizontalBox::Slot()
 				.AutoWidth()
@@ -111,8 +121,8 @@ namespace HLODOutliner
 				[
 					SNew(SButton)
 					.HAlign(HAlign_Center)
-					.Text(LOCTEXT("BuildHLODS", "Build HLODs"))
-					.OnClicked(this, &SHLODOutliner::HandleBuildHLODs)
+					.Text(LOCTEXT("GenerateClusters", "Generate Clusters"))
+					.OnClicked(this, &SHLODOutliner::HandlePreviewHLODs)
 				]
 
 				+ SHorizontalBox::Slot()
@@ -120,23 +130,11 @@ namespace HLODOutliner
 					.VAlign(VAlign_Center)
 					.Padding(FMargin(5.0f, 0.0f))
 					[
-						SNew(SButton)						
+						SNew(SButton)
 						.HAlign(HAlign_Center)
-						.Text(LOCTEXT("DeleteHLODS", "Delete HLODs"))
+						.Text(LOCTEXT("DeleteClusters", "Delete Clusters"))
 						.OnClicked(this, &SHLODOutliner::HandleDeleteHLODs)
 					]
-
-				+ SHorizontalBox::Slot()
-					.AutoWidth()
-					.VAlign(VAlign_Center)
-					.Padding(FMargin(5.0f, 0.0f))
-					[
-						SNew(SButton)						
-						.HAlign(HAlign_Center)
-						.Text(LOCTEXT("PreviewHLODS", "Preview HLODs"))
-						.OnClicked(this, &SHLODOutliner::HandlePreviewHLODs)
-					]
-
 				+ SHorizontalBox::Slot()
 					.AutoWidth()
 					.VAlign(VAlign_Center)
@@ -144,40 +142,10 @@ namespace HLODOutliner
 					[
 						SNew(SButton)
 						.HAlign(HAlign_Center)
-						.Text(LOCTEXT("DeletePreviewHLODS", "Delete Preview HLODs"))
-						.OnClicked(this, &SHLODOutliner::HandleDeletePreviewHLODs)
-					]
+						.Text(LOCTEXT("BuildMeshes", "Build LOD Meshes for Clusters"))
+						.OnClicked(this, &SHLODOutliner::HandleBuildLODActors)
+					]	
 
-				
-			]
-			+ SVerticalBox::Slot()
-			.Padding(FMargin(0.0f, 5.0f))
-			[
-				SNew(SHorizontalBox)
-				+SHorizontalBox::Slot()
-				.AutoWidth()
-				.VAlign(VAlign_Center)
-				.Padding(FMargin(5.0f, 0.0f))
-				[
-					SNew(SButton)
-					.HAlign(HAlign_Center)
-					.Text(LOCTEXT("BuildLODActors", "Build meshes for LODActors"))
-					.OnClicked(this, &SHLODOutliner::HandleBuildLODActors)
-				]
-
-#if 0 
-
-				+ SHorizontalBox::Slot()
-				.AutoWidth()
-				.VAlign(VAlign_Center)
-				.Padding(FMargin(5.0f, 0.0f))
-				[
-					SNew(SButton)
-					.HAlign(HAlign_Center)
-					.Text(LOCTEXT("ForceRefresh", "Force refresh"))
-					.OnClicked(this, &SHLODOutliner::HandleForceRefresh)
-				]
-#endif
 			];
 		
 	}
@@ -371,6 +339,7 @@ namespace HLODOutliner
 	{
 		if (CurrentWorld)
 		{
+			DestroySelectionActors();
 			CurrentWorld->HierarchicalLODBuilder->BuildMeshesForLODActors();
 		}
 
@@ -385,6 +354,59 @@ namespace HLODOutliner
 		return FReply::Handled();
 	}
 
+	FReply SHLODOutliner::HandleTestFunction()
+	{
+		if (SelectedNodes.Num() > 0)
+		{
+			if (SelectedNodes[0]->GetTreeItemType() == ITreeItem::HierarchicalLODActor)
+			{
+				FLODActorItem* Item = static_cast<FLODActorItem*>(SelectedNodes[0].Get());
+				ALODActor* LODActor = Item->LODActor.Get();
+				
+				FBox BoundingBox = LODActor->GetComponentsBoundingBox(true);
+
+				AHierarchicalLODVolume* Volume = CurrentWorld->SpawnActor<AHierarchicalLODVolume>(AHierarchicalLODVolume::StaticClass(), FTransform(BoundingBox.GetCenter()));
+				Volume->BrushBuilder = NewObject<UCubeBuilder>();			
+
+				// this code builds a brush for the new actor
+				Volume->PreEditChange(NULL);
+
+				Volume->PolyFlags = 0;
+				Volume->Brush = NewObject<UModel>(Volume, NAME_None, RF_Transactional);
+				Volume->Brush->Initialize(nullptr, true);
+				Volume->Brush->Polys = NewObject<UPolys>(Volume->Brush, NAME_None, RF_Transactional);
+				Volume->GetBrushComponent()->Brush = Volume->Brush;
+				Volume->BrushBuilder = NewObject<UCubeBuilder>();
+
+				UCubeBuilder* CubeBuilder = static_cast<UCubeBuilder*>(Volume->BrushBuilder);	
+
+				CubeBuilder->X = BoundingBox.GetSize().X * 1.5f;
+				CubeBuilder->Y = BoundingBox.GetSize().Y * 1.5f;
+				CubeBuilder->Z = BoundingBox.GetSize().Z * 1.5f;
+
+				Volume->BrushBuilder->Build(CurrentWorld, Volume);
+
+				FBSPOps::csgPrepMovingBrush(Volume);
+
+				// Set the texture on all polys to NULL.  This stops invisible textures
+				// dependencies from being formed on volumes.
+				if (Volume->Brush)
+				{
+					for (int32 poly = 0; poly < Volume->Brush->Polys->Element.Num(); ++poly)
+					{
+						FPoly* Poly = &(Volume->Brush->Polys->Element[poly]);
+						Poly->Material = NULL;
+					}
+				}
+
+				Volume->PostEditChange();
+
+				bool check = true;
+			}
+		}
+		return FReply::Handled();
+	}
+
 	END_SLATE_FUNCTION_BUILD_OPTIMIZATION
 
 		void SHLODOutliner::RegisterDelegates()
@@ -410,6 +432,8 @@ namespace HLODOutliner
 		GEngine->OnHLODActorAdded().AddSP(this, &SHLODOutliner::OnHLODActorAddedEvent);
 		GEngine->OnHLODActorMarkedDirty().AddSP(this, &SHLODOutliner::OnHLODActorMarkedDirtyEvent);
 		GEngine->OnHLODDrawDistanceChanged().AddSP(this, &SHLODOutliner::OnHLODDrawDistanceChangedEvent);
+		GEngine->OnHLODLevelsArrayChanged().AddSP(this, &SHLODOutliner::OnHLODLevelsArrayChangedEvent);
+		
 	}
 
 	void SHLODOutliner::DeregisterDelegates()
@@ -423,7 +447,7 @@ namespace HLODOutliner
 		GEngine->OnLevelActorDeleted().RemoveAll(this);
 		GEngine->OnActorMoved().RemoveAll(this);
 
-		FCoreDelegates::OnActorLabelChanged.AddRaw(this, &SHLODOutliner::OnActorLabelChanged);
+		FCoreDelegates::OnActorLabelChanged.RemoveAll(this);
 
 		USelection::SelectionChangedEvent.RemoveAll(this);
 		USelection::SelectObjectEvent.RemoveAll(this);
@@ -685,6 +709,14 @@ namespace HLODOutliner
 		}
 	}
 
+	void SHLODOutliner::RemoveLODLevelActors(const int32 HLODLevelIndex)
+	{
+		if (CurrentWorld)
+		{
+			HierarchicalLODUtils::DeleteLODActorsInHLODLevel(CurrentWorld, HLODLevelIndex);
+		}
+	}
+
 	TSharedRef<ITableRow> SHLODOutliner::OnOutlinerGenerateRow(FTreeItemPtr InTreeItem, const TSharedRef<STableViewBase>& OwnerTable)
 	{
 		TSharedRef<ITableRow> Widget = SNew(SHLODWidgetItem, OwnerTable)
@@ -1024,9 +1056,8 @@ namespace HLODOutliner
 		{
 			auto WorldSettings = CurrentWorld->GetWorldSettings();
 
-			check(WorldSettings->HierarchicalLODSetup.Num() == LODLevelDrawDistances.Num());
-
-			for (int32 LODLevelIndex = 0; LODLevelIndex < LODLevelDrawDistances.Num(); ++LODLevelIndex)
+			int32 MaxLODLevel = FMath::Min(WorldSettings->HierarchicalLODSetup.Num(), LODLevelDrawDistances.Num());
+			for (int32 LODLevelIndex = 0; LODLevelIndex < MaxLODLevel; ++LODLevelIndex)
 			{
 				if (LODLevelDrawDistances[LODLevelIndex] != WorldSettings->HierarchicalLODSetup[LODLevelIndex].DrawDistance)
 				{
@@ -1037,6 +1068,36 @@ namespace HLODOutliner
 		}
 	}
 
+	void SHLODOutliner::OnHLODLevelsArrayChangedEvent()
+	{
+		if (CurrentWorld)
+		{
+			auto WorldSettings = CurrentWorld->GetWorldSettings();			
+			if (WorldSettings->HierarchicalLODSetup.Num() > 1 && WorldSettings->HierarchicalLODSetup.Num() > LODLevelDrawDistances.Num())
+			{
+				// HLOD level was added, use previous level settings to "guess" new settings
+				auto& NewLevelSetup = WorldSettings->HierarchicalLODSetup.Last();
+				auto& OldLastLevelSetup = WorldSettings->HierarchicalLODSetup[WorldSettings->HierarchicalLODSetup.Num() - 2];
+
+				NewLevelSetup.bSimplifyMesh = OldLastLevelSetup.bSimplifyMesh;
+				NewLevelSetup.MergeSetting = OldLastLevelSetup.MergeSetting;
+				NewLevelSetup.ProxySetting = OldLastLevelSetup.ProxySetting;
+
+				NewLevelSetup.DesiredBoundRadius = OldLastLevelSetup.DesiredBoundRadius * 2.5f;
+				NewLevelSetup.DesiredFillingPercentage = FMath::Max( OldLastLevelSetup.DesiredFillingPercentage * 0.75f, 1.0f );
+				NewLevelSetup.DrawDistance = OldLastLevelSetup.DrawDistance * 2.5f;
+				NewLevelSetup.MinNumberOfActorsToBuild = OldLastLevelSetup.MinNumberOfActorsToBuild;			
+			}
+			else if (WorldSettings->HierarchicalLODSetup.Num() < LODLevelDrawDistances.Num())
+			{
+				// HLOD Level was removed, now remove all LODActors for this level
+				RemoveLODLevelActors(LODLevelDrawDistances.Num() - 1);
+			}
+
+			FullRefresh();
+		}
+	}
+
 	void SHLODOutliner::FullRefresh()
 	{
 		bNeedsRefresh = true;
diff --git a/Engine/Source/Editor/HierarchicalLODOutliner/Private/HLODOutliner.h b/Engine/Source/Editor/HierarchicalLODOutliner/Private/HLODOutliner.h
--- a/Engine/Source/Editor/HierarchicalLODOutliner/Private/HLODOutliner.h
+++ b/Engine/Source/Editor/HierarchicalLODOutliner/Private/HLODOutliner.h
@@ -99,6 +99,7 @@ namespace HLODOutliner
 		FReply RetrieveActors();
 		FReply HandleBuildLODActors();
 		FReply HandleForceRefresh();
+		FReply HandleTestFunction();
 		/** End button handlers */
 
 	private:
@@ -210,7 +211,13 @@ namespace HLODOutliner
 		* @param LODLevelIndex -
 		*/
 		void UpdateDrawDistancesForLODLevel(const uint32 LODLevelIndex);
-
+		
+		/**
+		* Removes LODActors within the given HLODLevel
+		*
+		* @param LODLevelIndex -		
+		*/
+		void RemoveLODLevelActors(const int32 HLODLevelIndex);
 	protected:
 		/** Tree view callbacks */
 
@@ -346,6 +353,9 @@ namespace HLODOutliner
 		/** Called when a DrawDistance value within WorldSettings changed */
 		void OnHLODDrawDistanceChangedEvent();
 
+		/** Called when the HLOD Levels array within WorldSettings changed */
+		void OnHLODLevelsArrayChangedEvent();
+
 		/** End of Broadcast event delegates */
 
 	private:
@@ -394,7 +404,7 @@ namespace HLODOutliner
 		* @param InItem - Item to select
 		*/
 		void SelectItemInTree(FTreeItemPtr InItem);
-	private:		
+	private:
 		/** Whether or not we need to do a refresh of the Tree view*/
 		bool bNeedsRefresh;
 
diff --git a/Engine/Source/Editor/HierarchicalLODOutliner/Private/LODLevelItem.cpp b/Engine/Source/Editor/HierarchicalLODOutliner/Private/LODLevelItem.cpp
--- a/Engine/Source/Editor/HierarchicalLODOutliner/Private/LODLevelItem.cpp
+++ b/Engine/Source/Editor/HierarchicalLODOutliner/Private/LODLevelItem.cpp
@@ -68,7 +68,28 @@ HLODOutliner::FDragValidationInfo HLODOutliner::FLODLevelDropTarget::ValidateDro
 
 		if (NumLODActors > 1)
 		{
-			return FDragValidationInfo(FHLODOutlinerDragDropOp::ToolTip_MultipleSelection_CompatibleNewCluster, LOCTEXT("CreateNewCluster", "Create new Cluster"));
+			// Check if all the dragged LOD actors fall within the same LOD level
+			auto LODActors = DraggedObjects.LODActors.GetValue();
+			int32 LevelIndex = -1;
+			bool bSameLODLevel = true;
+			for (auto Actor : LODActors)
+			{
+				ALODActor* LODActor = Cast<ALODActor>(Actor.Get());
+				if (LevelIndex == -1)
+				{
+					LevelIndex = LODActor->LODLevel;
+				}
+				else if (LevelIndex != LODActor->LODLevel)
+				{
+					bSameLODLevel = false;
+				}
+			}
+
+			if ( bSameLODLevel && LevelIndex < (int32)( LODLevelIndex + 1 ) )
+			{
+				return FDragValidationInfo(FHLODOutlinerDragDropOp::ToolTip_MultipleSelection_CompatibleNewCluster, LOCTEXT("CreateNewCluster", "Create new Cluster"));
+			}
+			
 		}
 	}
 
diff --git a/Engine/Source/Editor/HierarchicalLODOutliner/Private/StaticMeshActorItem.cpp b/Engine/Source/Editor/HierarchicalLODOutliner/Private/StaticMeshActorItem.cpp
--- a/Engine/Source/Editor/HierarchicalLODOutliner/Private/StaticMeshActorItem.cpp
+++ b/Engine/Source/Editor/HierarchicalLODOutliner/Private/StaticMeshActorItem.cpp
@@ -24,12 +24,19 @@ bool HLODOutliner::FStaticMeshActorItem::CanInteract() const
 void HLODOutliner::FStaticMeshActorItem::GenerateContextMenu(FMenuBuilder& MenuBuilder, SHLODOutliner& Outliner)
 {
 	auto SharedOutliner = StaticCastSharedRef<SHLODOutliner>(Outliner.AsShared());
-	MenuBuilder.AddMenuEntry(LOCTEXT("RemoveSMActorFromCluster", "Remove from cluster"), FText(), FSlateIcon(), FUIAction(FExecuteAction::CreateSP(&Outliner, &SHLODOutliner::RemoveStaticMeshActorFromCluster, AsShared())));
+	MenuBuilder.AddMenuEntry(LOCTEXT("RemoveSMActorFromCluster", "Remove From Cluster"), FText(), FSlateIcon(), FUIAction(FExecuteAction::CreateSP(&Outliner, &SHLODOutliner::RemoveStaticMeshActorFromCluster, AsShared())));
 }
 
 FString HLODOutliner::FStaticMeshActorItem::GetDisplayString() const
 {
-	return StaticMeshActor->GetFName().GetPlainNameString();
+	if (StaticMeshActor.IsValid())
+	{
+		return StaticMeshActor->GetFName().GetPlainNameString();
+	}
+	else
+	{
+		return FString("");
+	}	
 }
 
 HLODOutliner::FTreeItemID HLODOutliner::FStaticMeshActorItem::GetID()
diff --git a/Engine/Source/Editor/UnrealEd/Public/HierarchicalLODUtils.h b/Engine/Source/Editor/UnrealEd/Public/HierarchicalLODUtils.h
--- a/Engine/Source/Editor/UnrealEd/Public/HierarchicalLODUtils.h
+++ b/Engine/Source/Editor/UnrealEd/Public/HierarchicalLODUtils.h
@@ -328,6 +328,27 @@ namespace HierarchicalLODUtils
 		return NewActor;		
 	}
 
+	static void DeleteLODActorsInHLODLevel(UWorld* InWorld, const int32 HLODLevelIndex)
+	{
+		// you still have to delete all objects just in case they had it and didn't want it anymore
+		TArray<UObject*> AssetsToDelete;
+		for (int32 ActorId = InWorld->PersistentLevel->Actors.Num() - 1; ActorId >= 0; --ActorId)
+		{
+			ALODActor* LodActor = Cast<ALODActor>(InWorld->PersistentLevel->Actors[ActorId]);
+			if (LodActor && LodActor->LODLevel == ( HLODLevelIndex + 1))
+			{
+				DeleteLODActor(LodActor);
+				InWorld->DestroyActor(LodActor);
+			}
+		}
+	}
+
+	
+	static void TestBuild(ALODActor* LODActor, UPackage* AssetsOuter, const FHierarchicalSimplification& LODSetup, const uint32 LODIndex)
+	{
+
+	}
+
 };
 
 #undef LOCTEXT_NAMESPACE
\ No newline at end of file
diff --git a/Engine/Source/Runtime/Engine/Classes/Engine/Engine.h b/Engine/Source/Runtime/Engine/Classes/Engine/Engine.h
--- a/Engine/Source/Runtime/Engine/Classes/Engine/Engine.h
+++ b/Engine/Source/Runtime/Engine/Classes/Engine/Engine.h
@@ -1632,6 +1632,13 @@ public:
 	/** Called by internal engine systems after a HLOD Actor is marked dirty */
 	void BroadcastHLODDrawDistanceChanged() { HLODDrawDistanceChangedEvent.Broadcast(); }
 
+	/** Editor-only event triggered when a HLOD level is added or removed */
+	DECLARE_EVENT(UEngine, FHLODLevelsArrayChangedEvent);
+	FHLODLevelsArrayChangedEvent& OnHLODLevelsArrayChanged() { return HLODLevelsArrayChangedEvent; }
+
+	/** Called by internal engine systems after a HLOD Actor is marked dirty */
+	void BroadcastHLODLevelsArrayChanged() { HLODLevelsArrayChangedEvent.Broadcast(); }
+
 #endif // #if WITH_EDITOR
 
 	/** Event triggered after a server travel failure of any kind has occurred */
@@ -2344,6 +2351,9 @@ private:
 	/** Broadcasts after a Draw distance value (World settings) is changed */
 	FHLODDrawDistanceChangedEvent HLODDrawDistanceChangedEvent;
 
+	/** Broadcasts after the HLOD levels array is changed */
+	FHLODLevelsArrayChangedEvent HLODLevelsArrayChangedEvent;
+
 #endif // #if WITH_EDITOR
 
 	/** Thread preventing screen saver from kicking. Suspend most of the time. */
diff --git a/Engine/Source/Runtime/Engine/Private/LODActor.cpp b/Engine/Source/Runtime/Engine/Private/LODActor.cpp
--- a/Engine/Source/Runtime/Engine/Private/LODActor.cpp
+++ b/Engine/Source/Runtime/Engine/Private/LODActor.cpp
@@ -98,7 +98,7 @@ void ALODActor::PostEditChangeProperty(FPropertyChangedEvent& PropertyChangedEve
 bool ALODActor::GetReferencedContentObjects( TArray<UObject*>& Objects ) const
 {
 	Super::GetReferencedContentObjects(Objects);
-
+	Objects.Append(SubObjects);
 	return true;
 }
 
diff --git a/Engine/Source/Runtime/Engine/Private/WorldSettings.cpp b/Engine/Source/Runtime/Engine/Private/WorldSettings.cpp
--- a/Engine/Source/Runtime/Engine/Private/WorldSettings.cpp
+++ b/Engine/Source/Runtime/Engine/Private/WorldSettings.cpp
@@ -303,6 +303,11 @@ void AWorldSettings::PostEditChangeProperty(FPropertyChangedEvent& PropertyChang
 		{
 			GEditor->BroadcastHLODDrawDistanceChanged();
 		}
+
+		if (PropertyThatChanged->GetName() == TEXT("HierarchicalLODSetup"))
+		{
+			GEditor->BroadcastHLODLevelsArrayChanged();
+		}
 	}
 
 	Super::PostEditChangeProperty(PropertyChangedEvent);
