UE-19601 HLOD : Set up better default setting whenever adding new element
UE-18498 HLOD: Add functionality for choosing appropriate LOD from source StaticMesh when creating HLOD meshes
Fix for empty LOD actor not getting removed
Added context menu entry to remove SMA node from cluster and exclude from cluster generation in the future (setting the Auto LODGeneration flag on the actor to true)
Added feedback to Outliner when a Cluster contains less actors than the value of ôMin Number of Actors to buildö for the HLOD level it is in
Added option to create HLOD volume containing the selected clusters actors, for persistent changes in clustering generation
Added experimental slider for forcing the view Hierarchical LOD level

diff --git a/Engine/Source/Developer/MeshUtilities/Private/MeshUtilities.cpp b/Engine/Source/Developer/MeshUtilities/Private/MeshUtilities.cpp
--- a/Engine/Source/Developer/MeshUtilities/Private/MeshUtilities.cpp
+++ b/Engine/Source/Developer/MeshUtilities/Private/MeshUtilities.cpp
@@ -19,6 +19,8 @@
 #include "PhysicsEngine/BodySetup.h"
 #include "MaterialUtilities.h"
 
+#include "HierarchicalLODUtils.h"
+
 //@todo - implement required vector intrinsics for other implementations
 #if PLATFORM_ENABLE_VECTORINTRINSICS
 #include "kDOP.h"
@@ -4673,13 +4675,19 @@ void FMeshUtilities::MergeStaticMeshComponents(const TArray<UStaticMeshComponent
 		for (int32 MeshId = 0; MeshId < ComponentsToMerge.Num(); ++MeshId)
 		{
 			UStaticMeshComponent* MeshComponent = ComponentsToMerge[MeshId];
-
-			TArray<int32> MeshMaterialMap;
-			FRawMesh& RawMeshLOD = SourceMeshes[MeshId].MeshLOD[RawMeshLODIdx];
-
+						
 			// We duplicate lower LOD in case this mesh has no LOD we want
 			int32 ExportLODIndex = FMath::Min(LODIndex, MeshComponent->StaticMesh->SourceModels.Num() - 1);
 
+			// Determining which Source LOD level to pick according to the viewing distance
+			if (UseLOD == -1 && ViewDistance < TNumericLimits<float>::Max())
+			{
+				ExportLODIndex = HierarchicalLODUtils::GetLODLevelForDrawDistance(MeshComponent, ViewDistance);
+			}
+
+			TArray<int32> MeshMaterialMap;
+			FRawMesh& RawMeshLOD = SourceMeshes[MeshId].MeshLOD[RawMeshLODIdx];			
+
 			if (ConstructRawMesh(MeshComponent, ExportLODIndex, RawMeshLOD, UniqueMaterials, MeshMaterialMap))
 			{
 				MaterialMap.Add(FMeshIdAndLOD(MeshId, RawMeshLODIdx), MeshMaterialMap);
diff --git a/Engine/Source/Editor/HierarchicalLODOutliner/Private/HLODOutliner.cpp b/Engine/Source/Editor/HierarchicalLODOutliner/Private/HLODOutliner.cpp
--- a/Engine/Source/Editor/HierarchicalLODOutliner/Private/HLODOutliner.cpp
+++ b/Engine/Source/Editor/HierarchicalLODOutliner/Private/HLODOutliner.cpp
@@ -82,6 +82,14 @@ namespace HLODOutliner
 				CreateButtonWidgets()
 			];
 
+
+		MainContentPanel->AddSlot()
+			.AutoHeight()
+			.Padding(0.0f, 0.0f, 0.0f, 4.0f)
+			[				
+				CreateForcedViewSlider()
+			];
+
 		MainContentPanel->AddSlot()
 			.FillHeight(1.0f)
 			[
@@ -126,26 +134,26 @@ namespace HLODOutliner
 				]
 
 				+ SHorizontalBox::Slot()
-					.AutoWidth()
-					.VAlign(VAlign_Center)
-					.Padding(FMargin(5.0f, 0.0f))
-					[
-						SNew(SButton)
-						.HAlign(HAlign_Center)
-						.Text(LOCTEXT("DeleteClusters", "Delete Clusters"))
-						.OnClicked(this, &SHLODOutliner::HandleDeleteHLODs)
-					]
-				+ SHorizontalBox::Slot()
-					.AutoWidth()
-					.VAlign(VAlign_Center)
-					.Padding(FMargin(5.0f, 0.0f))
-					[
-						SNew(SButton)
-						.HAlign(HAlign_Center)
-						.Text(LOCTEXT("BuildMeshes", "Build LOD Meshes for Clusters"))
-						.OnClicked(this, &SHLODOutliner::HandleBuildLODActors)
-					]	
+				.AutoWidth()
+				.VAlign(VAlign_Center)
+				.Padding(FMargin(5.0f, 0.0f))
+				[
+					SNew(SButton)
+					.HAlign(HAlign_Center)
+					.Text(LOCTEXT("DeleteClusters", "Delete Clusters"))
+					.OnClicked(this, &SHLODOutliner::HandleDeleteHLODs)
+				]
 
+				+ SHorizontalBox::Slot()
+				.AutoWidth()
+				.VAlign(VAlign_Center)
+				.Padding(FMargin(5.0f, 0.0f))
+				[
+					SNew(SButton)
+					.HAlign(HAlign_Center)
+					.Text(LOCTEXT("BuildMeshes", "Build LOD Meshes for Clusters"))
+					.OnClicked(this, &SHLODOutliner::HandleBuildLODActors)
+				]
 			];
 		
 	}
@@ -177,6 +185,33 @@ namespace HLODOutliner
 				.FillWidth(0.25f)				
 				);
 	}
+
+	TSharedRef<SWidget> SHLODOutliner::CreateForcedViewSlider()
+	{
+		return SNew(SVerticalBox)
+			+ SVerticalBox::Slot()
+			.Padding(FMargin(0.0f, 5.0f))
+			[
+				SNew(SHorizontalBox)
+				+ SHorizontalBox::Slot()
+				.Padding(FMargin(5.0f, 0.0f))
+				.FillWidth(0.5f)
+				[
+					SNew(STextBlock)
+					.Text(this, &SHLODOutliner::HandleForceLevelText)
+				]
+				+ SHorizontalBox::Slot()
+					.FillWidth(0.5f)
+					[
+						SNew(SSlider)
+						.IsEnabled(this, &SHLODOutliner::HandleForcedLevelSliderIsEnabled)
+						.OnValueChanged(this, &SHLODOutliner::HandleForcedLevelSliderValueChanged)
+						.Orientation(Orient_Horizontal)
+						.Value(this, &SHLODOutliner::HandleForcedLevelSliderValue)
+					]
+			];
+	}
+
 	void SHLODOutliner::CreateSettingsView()
 	{
 		// Create a property view
@@ -202,7 +237,7 @@ namespace HLODOutliner
 			/** Delegate to show all properties */
 			static bool IsPropertyVisible(const FPropertyAndParent& PropertyAndParent, bool bInShouldShowNonEditable)
 			{
-				if (PropertyAndParent.Property.GetName() == "HierarchicalLODSetup" || (PropertyAndParent.ParentProperty && PropertyAndParent.ParentProperty->GetName() == "MergeSetting"))
+				if (PropertyAndParent.Property.GetName() == "HierarchicalLODSetup" || (PropertyAndParent.ParentProperty && PropertyAndParent.ParentProperty->GetName() == "MergeSetting") || (PropertyAndParent.ParentProperty && PropertyAndParent.ParentProperty->GetName() == "ProxySetting"))
 				{
 					return true;
 				}
@@ -363,45 +398,7 @@ namespace HLODOutliner
 				FLODActorItem* Item = static_cast<FLODActorItem*>(SelectedNodes[0].Get());
 				ALODActor* LODActor = Item->LODActor.Get();
 				
-				FBox BoundingBox = LODActor->GetComponentsBoundingBox(true);
-
-				AHierarchicalLODVolume* Volume = CurrentWorld->SpawnActor<AHierarchicalLODVolume>(AHierarchicalLODVolume::StaticClass(), FTransform(BoundingBox.GetCenter()));
-				Volume->BrushBuilder = NewObject<UCubeBuilder>();			
-
-				// this code builds a brush for the new actor
-				Volume->PreEditChange(NULL);
-
-				Volume->PolyFlags = 0;
-				Volume->Brush = NewObject<UModel>(Volume, NAME_None, RF_Transactional);
-				Volume->Brush->Initialize(nullptr, true);
-				Volume->Brush->Polys = NewObject<UPolys>(Volume->Brush, NAME_None, RF_Transactional);
-				Volume->GetBrushComponent()->Brush = Volume->Brush;
-				Volume->BrushBuilder = NewObject<UCubeBuilder>();
-
-				UCubeBuilder* CubeBuilder = static_cast<UCubeBuilder*>(Volume->BrushBuilder);	
-
-				CubeBuilder->X = BoundingBox.GetSize().X * 1.5f;
-				CubeBuilder->Y = BoundingBox.GetSize().Y * 1.5f;
-				CubeBuilder->Z = BoundingBox.GetSize().Z * 1.5f;
-
-				Volume->BrushBuilder->Build(CurrentWorld, Volume);
-
-				FBSPOps::csgPrepMovingBrush(Volume);
-
-				// Set the texture on all polys to NULL.  This stops invisible textures
-				// dependencies from being formed on volumes.
-				if (Volume->Brush)
-				{
-					for (int32 poly = 0; poly < Volume->Brush->Polys->Element.Num(); ++poly)
-					{
-						FPoly* Poly = &(Volume->Brush->Polys->Element[poly]);
-						Poly->Material = NULL;
-					}
-				}
-
-				Volume->PostEditChange();
-
-				bool check = true;
+				HierarchicalLODUtils::CreateVolumeForLODActor(LODActor, CurrentWorld);
 			}
 		}
 		return FReply::Handled();
@@ -499,6 +496,39 @@ namespace HLODOutliner
 		}
 	}
 
+	bool SHLODOutliner::HandleForcedLevelSliderIsEnabled() const
+	{
+		bool bHLODsBuild = true;
+		for (bool Build : LODLevelBuildFlags)
+		{
+			bHLODsBuild &= Build;
+		}
+
+		return (LODLevelDrawDistances.Num() > 0 && bHLODsBuild);
+	}
+
+	void SHLODOutliner::HandleForcedLevelSliderValueChanged(float NewValue)
+	{
+		// Snap values
+		int32 SnappedValue = FMath::Min( NewValue, 1.0f ) * LODLevelDrawDistances.Num();
+		if (SnappedValue - 1 != ForcedLODLevel)
+		{
+			RestoreForcedLODLevel(ForcedLODLevel);
+			ForcedLODLevel = -1;
+			SetForcedLODLevel(SnappedValue - 1);
+		}			
+	}
+
+	float SHLODOutliner::HandleForcedLevelSliderValue() const
+	{
+		return ((1.0f / (LODLevelDrawDistances.Num())) * (ForcedLODLevel + 1));
+	}
+
+	FText SHLODOutliner::HandleForceLevelText() const
+	{
+		return FText::FromString(FString("Forced viewing level: ") + ( (ForcedLODLevel == -1) ? FString("None") : FString::FromInt(ForcedLODLevel)));
+	}
+
 	bool SHLODOutliner::CanHLODLevelBeForced(const uint32 LODLevel) const
 	{
 		return LODLevelBuildFlags[LODLevel];
@@ -534,6 +564,12 @@ namespace HLODOutliner
 
 	void SHLODOutliner::SetForcedLODLevel(const uint32 LODLevel)
 	{
+		if (LODLevel == -1)
+		{
+			ForcedLODLevel = LODLevel;
+			return;
+		}
+
 		if (CurrentWorld)
 		{
 			auto Level = CurrentWorld->GetCurrentLevel();
@@ -556,6 +592,18 @@ namespace HLODOutliner
 		ForcedLODLevel = LODLevel;
 	}
 
+	void SHLODOutliner::CreateHierarchicalVolumeForActor(TSharedRef<ITreeItem> Item)
+	{
+		if (Item->GetTreeItemType() == ITreeItem::HierarchicalLODActor)
+		{
+			FLODActorItem* ActorItem = (FLODActorItem*)(&Item.Get());
+			ALODActor* LODActor = ActorItem->LODActor.Get();
+
+			AHierarchicalLODVolume* Volume = HierarchicalLODUtils::CreateVolumeForLODActor(LODActor, CurrentWorld);
+			check(Volume);
+		}
+	}
+
 	void SHLODOutliner::BuildLODActor(TSharedRef<ITreeItem> Item)
 	{
 		if (CurrentWorld)
@@ -641,6 +689,15 @@ namespace HLODOutliner
 		}
 	}
 
+	void SHLODOutliner::ExcludeFromClusterGeneration(TSharedRef<ITreeItem> Item)
+	{
+		const FScopedTransaction Transaction(LOCTEXT("UndoAction_ExcludeStaticMeshActorFromClusterGeneration", "Excluded StaticMeshActor From Cluster Generation"));
+		FStaticMeshActorItem* ActorItem = (FStaticMeshActorItem*)(&Item.Get());
+		ActorItem->StaticMeshActor->Modify();
+		ActorItem->StaticMeshActor->bEnableAutoLODGeneration = false;
+		RemoveStaticMeshActorFromCluster(Item);		
+	}
+
 	void SHLODOutliner::RemoveLODActorFromCluster(TSharedRef<ITreeItem> Item)
 	{
 		if (CurrentWorld)
@@ -703,6 +760,8 @@ namespace HLODOutliner
 					{
 						LODActor->LODDrawDistance = LODLevelDrawDistances[LODLevelIndex];
 						LODActor->UpdateSubActorLODParents();
+						
+						LODActor->SetIsDirty(true);
 					}
 				}
 			}
diff --git a/Engine/Source/Editor/HierarchicalLODOutliner/Private/HLODOutliner.h b/Engine/Source/Editor/HierarchicalLODOutliner/Private/HLODOutliner.h
--- a/Engine/Source/Editor/HierarchicalLODOutliner/Private/HLODOutliner.h
+++ b/Engine/Source/Editor/HierarchicalLODOutliner/Private/HLODOutliner.h
@@ -68,6 +68,10 @@ namespace HLODOutliner
 		* Constructs this widget.
 		*/
 		void Construct(const FArguments& InArgs);
+
+		/** Create the panel's forced HLOD level viewer */
+		TSharedRef<SWidget> CreateForcedViewSlider();
+
 		/** Creates the panel's ButtonWidget rows */
 		TSharedRef<SWidget> CreateButtonWidgets();
 
@@ -133,6 +137,34 @@ namespace HLODOutliner
 		*/
 		void HandleCheckBoxCheckedStateChanged(ECheckBoxState NewState, TSharedRef<ITreeItem> Item);
 
+		/**
+		* Determines if the HLOD slider is enabled, only if all HLODs are buil
+		*
+		* @return bool
+		*/
+		bool HandleForcedLevelSliderIsEnabled() const;
+		
+		/**
+		* Handles changes in slider value, maps/snaps it to the corresponding HLOD level
+		*
+		* @param NewValue -
+		*/
+		void HandleForcedLevelSliderValueChanged(float NewValue);
+
+		/**
+		* Returns the float sliding value corresponding to the currently forced HLOD level
+		*
+		* @return float
+		*/
+		float HandleForcedLevelSliderValue() const;
+
+		/**
+		* Returns FText with information about current forced HLOD level "None", "1" etc.
+		*
+		* @return FText
+		*/
+		FText HandleForceLevelText() const;
+
 		/**
 		* Returns whether or not a HLOD level can be forced, this depends on whether or not all of the Clusters with LODLevel are non-dirty (have their meshes build)
 		*
@@ -155,6 +187,12 @@ namespace HLODOutliner
 		*/
 		void SetForcedLODLevel(const uint32 LODLevel);
 
+		/**
+		* Creates a Hierarchical LOD Volume for the given LODActorItem, volume bounds correspond to those of the LODActor's SubActors
+		*
+		* @param Item - LODActorItem
+		*/
+		void CreateHierarchicalVolumeForActor(TSharedRef<ITreeItem> Item);
 	protected:
 		/**
 		* Builds the HLOD mesh for the given ALODActor (cluster)
@@ -190,6 +228,13 @@ namespace HLODOutliner
 		* @param Item - (StaticMeshActor)TreeItem containing the Actor ptr
 		*/
 		void RemoveStaticMeshActorFromCluster(TSharedRef<ITreeItem> Item);
+		
+		/**
+		* Removes the given StaticMeshActor from its parent's (ALODActor) sub-actors array and excludes it from cluster generation
+		*
+		* @param Item - (StaticMeshActor)TreeItem containing the Actor ptr
+		*/
+		void ExcludeFromClusterGeneration(TSharedRef<ITreeItem> Item);		
 
 		/**
 		* Removes the given LODActor from its parent's (ALODActor) sub-actors array
diff --git a/Engine/Source/Editor/HierarchicalLODOutliner/Private/LODActorItem.cpp b/Engine/Source/Editor/HierarchicalLODOutliner/Private/LODActorItem.cpp
--- a/Engine/Source/Editor/HierarchicalLODOutliner/Private/LODActorItem.cpp
+++ b/Engine/Source/Editor/HierarchicalLODOutliner/Private/LODActorItem.cpp
@@ -40,6 +40,8 @@ void HLODOutliner::FLODActorItem::GenerateContextMenu(FMenuBuilder& MenuBuilder,
 		MenuBuilder.AddMenuEntry(LOCTEXT("ForceView", "ForceView"), FText(), FSlateIcon(), FUIAction(FExecuteAction::CreateSP(&Outliner, &SHLODOutliner::ForceViewLODActor, AsShared())));
 	}
 
+	MenuBuilder.AddMenuEntry(LOCTEXT("CreateHLODVolume", "Create Containing Hierarchical Volume"), FText(), FSlateIcon(), FUIAction(FExecuteAction::CreateSP(&Outliner, &SHLODOutliner::CreateHierarchicalVolumeForActor, AsShared())));
+
 	AActor* Actor = LODActor.Get();
 	ALODActor* ParentActor = HierarchicalLODUtils::GetParentLODActor(Actor);
 	if (ParentActor && Parent.Pin()->GetTreeItemType() == TreeItemType::HierarchicalLODActor)
@@ -56,7 +58,7 @@ FString HLODOutliner::FLODActorItem::GetDisplayString() const
 {
 	if (ALODActor* ActorPtr = LODActor.Get())
 	{
-		return ActorPtr->GetName() + ((ActorPtr->IsDirty()) ? FString(" (Not built)") : FString());
+		return ActorPtr->GetName() + ((ActorPtr->SubActors.Num() < 2) ? FString(" (Insufficient Sub Actors)") : ((ActorPtr->IsDirty()) ? FString(" (Not built)") : FString()));
 	}
 
 	return FString();
diff --git a/Engine/Source/Editor/HierarchicalLODOutliner/Private/LODLevelItem.cpp b/Engine/Source/Editor/HierarchicalLODOutliner/Private/LODLevelItem.cpp
--- a/Engine/Source/Editor/HierarchicalLODOutliner/Private/LODLevelItem.cpp
+++ b/Engine/Source/Editor/HierarchicalLODOutliner/Private/LODLevelItem.cpp
@@ -119,31 +119,46 @@ void HLODOutliner::FLODLevelDropTarget::CreateNewCluster(FDragDropPayload &Dragg
 
 	const FScopedTransaction Transaction(LOCTEXT("UndoAction_CreateNewCluster", "Create new Cluster"));
 	World->Modify();
-	
-	ALODActor* NewCluster = HierarchicalLODUtils::CreateNewClusterActor(World, LODLevelIndex);
 
-	for (TWeakObjectPtr<AActor> StaticMeshActor : DraggedObjects.StaticMeshActors.GetValue())
+	if (World->GetWorldSettings()->bEnableHierarchicalLODSystem)
 	{
-		AActor* InActor = StaticMeshActor.Get();
-		ALODActor* CurrentParentActor = HierarchicalLODUtils::GetParentLODActor(InActor);
-		if (CurrentParentActor)
+		ALODActor* NewCluster = HierarchicalLODUtils::CreateNewClusterActor(World, LODLevelIndex);
+
+		if (NewCluster)
 		{
-			CurrentParentActor->RemoveSubActor(InActor);
-		}
+			for (TWeakObjectPtr<AActor> StaticMeshActor : DraggedObjects.StaticMeshActors.GetValue())
+			{
+				AActor* InActor = StaticMeshActor.Get();
+				ALODActor* CurrentParentActor = HierarchicalLODUtils::GetParentLODActor(InActor);
+				if (CurrentParentActor)
+				{
+					CurrentParentActor->RemoveSubActor(InActor);
 
-		NewCluster->AddSubActor(InActor);
-	}
+					if (!CurrentParentActor->HasValidSubActors())
+					{
+						HierarchicalLODUtils::DeleteLODActor(CurrentParentActor);
+					}
+				}
 
-	for (TWeakObjectPtr<AActor> LODActor : DraggedObjects.LODActors.GetValue())
-	{
-		AActor* InActor = LODActor.Get();
-		ALODActor* CurrentParentActor = HierarchicalLODUtils::GetParentLODActor(InActor);
-		if (CurrentParentActor)
-		{
-			CurrentParentActor->RemoveSubActor(InActor);
-		}
+				NewCluster->AddSubActor(InActor);
+			}
+
+			for (TWeakObjectPtr<AActor> LODActor : DraggedObjects.LODActors.GetValue())
+			{
+				AActor* InActor = LODActor.Get();
+				ALODActor* CurrentParentActor = HierarchicalLODUtils::GetParentLODActor(InActor);
+				if (CurrentParentActor)
+				{
+					CurrentParentActor->RemoveSubActor(InActor);
+					if (!CurrentParentActor->HasValidSubActors())
+					{
+						HierarchicalLODUtils::DeleteLODActor(CurrentParentActor);
+					}
+				}
 
-		NewCluster->AddSubActor(InActor);
+				NewCluster->AddSubActor(InActor);
+			}
+		}
 	}
 }
 
diff --git a/Engine/Source/Editor/HierarchicalLODOutliner/Private/StaticMeshActorItem.cpp b/Engine/Source/Editor/HierarchicalLODOutliner/Private/StaticMeshActorItem.cpp
--- a/Engine/Source/Editor/HierarchicalLODOutliner/Private/StaticMeshActorItem.cpp
+++ b/Engine/Source/Editor/HierarchicalLODOutliner/Private/StaticMeshActorItem.cpp
@@ -25,6 +25,7 @@ void HLODOutliner::FStaticMeshActorItem::GenerateContextMenu(FMenuBuilder& MenuB
 {
 	auto SharedOutliner = StaticCastSharedRef<SHLODOutliner>(Outliner.AsShared());
 	MenuBuilder.AddMenuEntry(LOCTEXT("RemoveSMActorFromCluster", "Remove From Cluster"), FText(), FSlateIcon(), FUIAction(FExecuteAction::CreateSP(&Outliner, &SHLODOutliner::RemoveStaticMeshActorFromCluster, AsShared())));
+	MenuBuilder.AddMenuEntry(LOCTEXT("ExcludeSMActorFromGeneration", "Exclude from Cluster Generation"), FText(), FSlateIcon(), FUIAction(FExecuteAction::CreateSP(&Outliner, &SHLODOutliner::ExcludeFromClusterGeneration, AsShared())));
 }
 
 FString HLODOutliner::FStaticMeshActorItem::GetDisplayString() const
diff --git a/Engine/Source/Editor/UnrealEd/Private/LODCluster.cpp b/Engine/Source/Editor/UnrealEd/Private/LODCluster.cpp
--- a/Engine/Source/Editor/UnrealEd/Private/LODCluster.cpp
+++ b/Engine/Source/Editor/UnrealEd/Private/LODCluster.cpp
@@ -322,7 +322,7 @@ void FLODCluster::BuildActor(ULevel* InLevel, const int32 LODIdx, const bool bCr
 					}
 					else
 					{						
-						MeshUtilities.MergeStaticMeshComponents(AllComponents, FirstActor->GetWorld(), LODSetup.MergeSetting, AssetsOuter, PackageName, LODIdx, OutAssets, OutProxyLocation, LODSetup.DrawDistance, true);
+						MeshUtilities.MergeStaticMeshComponents(AllComponents, FirstActor->GetWorld(), LODSetup.MergeSetting, AssetsOuter, PackageName, -1, OutAssets, OutProxyLocation, LODSetup.DrawDistance, true);
 					}
 
 					// we make it private, so it can't be used by outside of map since it's useless, and then remove standalone
diff --git a/Engine/Source/Editor/UnrealEd/Public/HierarchicalLODUtils.h b/Engine/Source/Editor/UnrealEd/Public/HierarchicalLODUtils.h
--- a/Engine/Source/Editor/UnrealEd/Public/HierarchicalLODUtils.h
+++ b/Engine/Source/Editor/UnrealEd/Public/HierarchicalLODUtils.h
@@ -1,12 +1,22 @@
 // Copyright 1998-2015 Epic Games, Inc. All Rights Reserved.
 #pragma once
 
+#include "Core.h"
+#include "AsyncWork.h"
+
+#include "Engine/Engine.h"
+#include "Engine/World.h"
 #include "Engine/LODActor.h"
 #include "Components/StaticMeshComponent.h"
 #include "Engine/StaticMesh.h"
 #include "MeshUtilities.h"
+#include "StaticMeshResources.h"
+#include "HierarchicalLODVolume.h"
 
-#include "AsyncWork.h"
+#include "Editor.h"
+#include "Editor/EditorEngine.h"
+#include "BSPOps.h"
+#include "UnrealEd.h"
 
 #if WITH_EDITOR
 #include "AssetData.h"
@@ -57,7 +67,7 @@ namespace HierarchicalLODUtils
 		{
 			if (!LODActor->IsDirty() || LODActor->SubActors.Num() < 2)
 			{
-				return true;
+				return false;
 			}
 
 			// Clean up sub objects (if applicable)
@@ -159,6 +169,12 @@ namespace HierarchicalLODUtils
 	}
 
 
+	/**
+	* Returns whether or not the given actor is eligble for creating a HLOD cluster creation
+	*
+	* @param Actor -
+	* @return bool
+	*/
 	static bool ShouldGenerateCluster(AActor* Actor)
 	{
 		if (!Actor)
@@ -223,6 +239,7 @@ namespace HierarchicalLODUtils
 		return (ValidComponentCount > 0);
 	}
 
+	/** Returns the ALODActor parent for the given InActor, nullptr if none available */
 	static ALODActor* GetParentLODActor(const AActor* InActor)
 	{
 		ALODActor* ParentActor = nullptr;
@@ -244,6 +261,7 @@ namespace HierarchicalLODUtils
 		return ParentActor;		
 	}
 
+	/** Deletes the given actor's data and instance in the world */
 	static void DeleteLODActor(ALODActor* InActor)
 	{
 		// Find if it has a parent ALODActor
@@ -283,6 +301,7 @@ namespace HierarchicalLODUtils
 		CollectGarbage(GARBAGE_COLLECTION_KEEPFLAGS, true);
 	}
 
+	/** Deletes the given actor's assets */
 	static void DeleteLODActorAssets(ALODActor* InActor)
 	{
 		// you still have to delete all objects just in case they had it and didn't want it anymore
@@ -311,12 +330,16 @@ namespace HierarchicalLODUtils
 		CollectGarbage(GARBAGE_COLLECTION_KEEPFLAGS, true);
 	}
 
+	/** Creates a new cluster actor in the given InWorld with InLODLevel as HLODLevel */
 	static ALODActor* CreateNewClusterActor(UWorld* InWorld, const int32 InLODLevel)
 	{
 		// Check incoming data
 		check(InWorld != nullptr && InLODLevel >= 0);
 		auto WorldSettings = InWorld->GetWorldSettings();
-		check(WorldSettings->bEnableHierarchicalLODSystem && WorldSettings->HierarchicalLODSetup.Num() > InLODLevel);
+		if (!WorldSettings->bEnableHierarchicalLODSystem || WorldSettings->HierarchicalLODSetup.Num() == 0 || WorldSettings->HierarchicalLODSetup.Num() < InLODLevel)
+		{
+			return nullptr;
+		}
 					
 		// Spawn and setup actor
 		ALODActor* NewActor = nullptr;
@@ -328,6 +351,7 @@ namespace HierarchicalLODUtils
 		return NewActor;		
 	}
 
+	/** Deletes all the ALODActors with the given HLODLevelIndex inside off InWorld */
 	static void DeleteLODActorsInHLODLevel(UWorld* InWorld, const int32 HLODLevelIndex)
 	{
 		// you still have to delete all objects just in case they had it and didn't want it anymore
@@ -342,11 +366,92 @@ namespace HierarchicalLODUtils
 			}
 		}
 	}
+	
+	/** Computes the Screensize of the given Sphere taking into account the ProjectionMatrix and distance */
+	static float ComputeBoundsScreenSize(const float SphereRadius, const FMatrix& ProjectionMatrix, const float Distance)
+	{
+		// Only need one component from a view transformation; just calculate the one we're interested in.
+		const float Divisor = Distance;
+
+		// Get projection multiple accounting for view scaling.
+		const float ScreenMultiple = FMath::Max(1920.0f / 2.0f * ProjectionMatrix.M[0][0],
+			1080.0f / 2.0f * ProjectionMatrix.M[1][1]);
+
+		const float ScreenRadius = ScreenMultiple * SphereRadius / FMath::Max(Divisor, 1.0f);
+		const float ScreenArea = PI * ScreenRadius * ScreenRadius;
+		return FMath::Clamp(ScreenArea / (1920.0f * 1080.0f), 0.0f, 1.0f);
+	}
+	
+	/** Computes which LOD level of a Mesh corresponds to the given Distance (calculates closest screensize with distance) */
+	static int32 ComputeStaticMeshLODLevel(const TArray<FStaticMeshSourceModel>& SourceModels, const float SphereRadius, const FMatrix& ProjectionMatrix, const float Distance)
+	{
+		const int32 NumLODs = SourceModels.Num();
+
+		const float ScreenSize = ComputeBoundsScreenSize(SphereRadius, ProjectionMatrix, Distance);
+
+		// Walk backwards and return the first matching LOD
+		for (int32 LODIndex = NumLODs - 1; LODIndex >= 0; --LODIndex)
+		{
+			if (SourceModels[LODIndex].ScreenSize > ScreenSize)
+			{
+			
+				return FMath::Max(LODIndex, 0);
+			}
+		}
 
+		return 0;
+	}
 	
-	static void TestBuild(ALODActor* LODActor, UPackage* AssetsOuter, const FHierarchicalSimplification& LODSetup, const uint32 LODIndex)
+	/** Computes the LODLevel for a StaticMeshComponent taking into account the DrawingDistance */
+	static int32 GetLODLevelForDrawDistance(UStaticMeshComponent* StaticMeshComponent, const float DrawDistance)
+	{
+		// convert FOV to radians
+		float FOVRad = 90.0f * (float)PI / 360.0f;
+		FMatrix ProjectionMatrix = FPerspectiveMatrix(FOVRad, 1920, 1080, 0.01f);
+		FBoxSphereBounds Bounds = StaticMeshComponent->CalcBounds(FTransform());
+		return ComputeStaticMeshLODLevel(StaticMeshComponent->StaticMesh->SourceModels, Bounds.SphereRadius, ProjectionMatrix, DrawDistance);
+	}
+
+	static AHierarchicalLODVolume* CreateVolumeForLODActor(ALODActor* InLODActor, UWorld* InWorld )
 	{
+		FBox BoundingBox = InLODActor->GetComponentsBoundingBox(true);
+
+		AHierarchicalLODVolume* Volume = InWorld->SpawnActor<AHierarchicalLODVolume>(AHierarchicalLODVolume::StaticClass(), FTransform(BoundingBox.GetCenter()));
+
+		// this code builds a brush for the new actor
+		Volume->PreEditChange(NULL);
+
+		Volume->PolyFlags = 0;
+		Volume->Brush = NewObject<UModel>(Volume, NAME_None, RF_Transactional);
+		Volume->Brush->Initialize(nullptr, true);
+		Volume->Brush->Polys = NewObject<UPolys>(Volume->Brush, NAME_None, RF_Transactional);
+		Volume->GetBrushComponent()->Brush = Volume->Brush;
+		Volume->BrushBuilder = NewObject<UCubeBuilder>( Volume, NAME_None, RF_Transactional);
+
+		UCubeBuilder* CubeBuilder = static_cast<UCubeBuilder*>(Volume->BrushBuilder);
+
+		CubeBuilder->X = BoundingBox.GetSize().X * 1.5f;
+		CubeBuilder->Y = BoundingBox.GetSize().Y * 1.5f;
+		CubeBuilder->Z = BoundingBox.GetSize().Z * 1.5f;
+
+		Volume->BrushBuilder->Build(InWorld, Volume);
+
+		FBSPOps::csgPrepMovingBrush(Volume);
+
+		// Set the texture on all polys to NULL.  This stops invisible textures
+		// dependencies from being formed on volumes.
+		if (Volume->Brush)
+		{
+			for (int32 poly = 0; poly < Volume->Brush->Polys->Element.Num(); ++poly)
+			{
+				FPoly* Poly = &(Volume->Brush->Polys->Element[poly]);
+				Poly->Material = NULL;
+			}
+		}
+
+		Volume->PostEditChange();
 
+		return Volume;
 	}
 
 };
diff --git a/Engine/Source/Runtime/Engine/Private/LODActor.cpp b/Engine/Source/Runtime/Engine/Private/LODActor.cpp
--- a/Engine/Source/Runtime/Engine/Private/LODActor.cpp
+++ b/Engine/Source/Runtime/Engine/Private/LODActor.cpp
@@ -314,19 +314,19 @@ void ALODActor::SetHiddenFromEditorView(const bool InState, const int32 ForceLOD
 	// If we are also subactor for a higher LOD level or this actor belongs to a higher HLOD level than is being forced hide the actor
 	if (GetStaticMeshComponent()->GetLODParentPrimitive() || LODLevel > ForceLODLevel )
 	{
-		SetIsTemporarilyHiddenInEditor(InState);		
-	}
+		SetIsTemporarilyHiddenInEditor(InState);			
 
-	for (auto Actor : SubActors)
-	{
-		// If this actor belongs to a lower HLOD level that is being forced hide the sub-actors
-		if (LODLevel < ForceLODLevel)
+		for (auto Actor : SubActors)
 		{
-			Actor->SetIsTemporarilyHiddenInEditor(InState);
-		}
+			// If this actor belongs to a lower HLOD level that is being forced hide the sub-actors
+			if (LODLevel < ForceLODLevel)
+			{
+				Actor->SetIsTemporarilyHiddenInEditor(InState);
+			}
 
-		// Toggle/set the LOD parent to nullptr or this
-		Actor->SetLODParent((InState) ? nullptr : StaticMeshComponent, (InState) ? 0.0f : LODDrawDistance);
+			// Toggle/set the LOD parent to nullptr or this
+			Actor->SetLODParent((InState) ? nullptr : StaticMeshComponent, (InState) ? 0.0f : LODDrawDistance);
+		}
 	}
 
 	StaticMeshComponent->MarkRenderStateDirty();
