Changed behaviour of forced HLOD level slider (more use-friendly now)
Added console command for enabling/disabling the HLOD system
Changed/removed HLOD behaviour from build menu
Removed context menu from HLOD level treeview item
Double clicking a (build) LODActor or StaticMeshActor item in the HLOD outliner will now move the editor camera to the actor (just like the scene outliner)

diff --git a/Engine/Source/Editor/HierarchicalLODOutliner/Private/HLODOutliner.cpp b/Engine/Source/Editor/HierarchicalLODOutliner/Private/HLODOutliner.cpp
--- a/Engine/Source/Editor/HierarchicalLODOutliner/Private/HLODOutliner.cpp
+++ b/Engine/Source/Editor/HierarchicalLODOutliner/Private/HLODOutliner.cpp
@@ -1,4 +1,4 @@
-// Copyright 1998-2015 Epic Games, Inc. All Rights Reserved.
+// Copyirght 1998-2015 Epic Games, Inc. All Rights Reserved.
 
 #include "HierarchicalLODOutlinerPrivatePCH.h"
 #include "HLODOutliner.h"
@@ -27,7 +27,6 @@
 
 #include "HierarchicalLODUtils.h"
 
-
 #define LOCTEXT_NAMESPACE "HLODOutliner"
 
 namespace HLODOutliner
@@ -37,6 +36,8 @@ namespace HLODOutliner
 		bNeedsRefresh = true;
 		CurrentWorld = nullptr;
 		ForcedLODLevel = -1;
+		ForcedLODSliderValue = 0.0f;
+		bForcedSliderValueUpdating = false;
 	}
 
 	SHLODOutliner::~SHLODOutliner()
@@ -95,16 +96,16 @@ namespace HLODOutliner
 			[
 				SNew(SSplitter)
 				.Orientation(Orient_Vertical)
-
+				.Style(FEditorStyle::Get(), "ContentBrowser.Splitter")
 				+ SSplitter::Slot()
 				.Value(0.5)
 				[
-					CreateTreeviewWidget()
+					CreateTreeviewWidget()										
 				]
 				+ SSplitter::Slot()
 				.Value(0.5)
 				[
-					SettingsView.ToSharedRef()
+					SettingsView.ToSharedRef()										
 				]		
 			];
 
@@ -168,7 +169,7 @@ namespace HLODOutliner
 			.OnSelectionChanged(this, &SHLODOutliner::OnOutlinerSelectionChanged)
 			.OnMouseButtonDoubleClick(this, &SHLODOutliner::OnOutlinerDoubleClick)
 			.OnContextMenuOpening(this, &SHLODOutliner::OnOpenContextMenu)
-			.OnExpansionChanged(this, &SHLODOutliner::OnItemExpansionChanged)
+			.OnExpansionChanged(this, &SHLODOutliner::OnItemExpansionChanged)			
 			.HeaderRow
 			(
 				SNew(SHeaderRow)
@@ -201,14 +202,17 @@ namespace HLODOutliner
 					.Text(this, &SHLODOutliner::HandleForceLevelText)
 				]
 				+ SHorizontalBox::Slot()
-					.FillWidth(0.5f)
-					[
-						SNew(SSlider)
-						.IsEnabled(this, &SHLODOutliner::HandleForcedLevelSliderIsEnabled)
-						.OnValueChanged(this, &SHLODOutliner::HandleForcedLevelSliderValueChanged)
-						.Orientation(Orient_Horizontal)
-						.Value(this, &SHLODOutliner::HandleForcedLevelSliderValue)
-					]
+				.Padding(FMargin(5.0f, 0.0f))
+				.FillWidth(0.5f)
+				[
+					SNew(SSlider)
+					.IsEnabled(this, &SHLODOutliner::HandleForcedLevelSliderIsEnabled)
+					.OnValueChanged(this, &SHLODOutliner::HandleForcedLevelSliderValueChanged)
+					.OnMouseCaptureBegin(this, &SHLODOutliner::HandleForcedLevelSliderCaptureBegin)
+					.OnMouseCaptureEnd(this, &SHLODOutliner::HandleForcedLevelSliderCaptureEnd)
+					.Orientation(Orient_Horizontal)
+					.Value(this, &SHLODOutliner::HandleForcedLevelSliderValue)
+				]
 			];
 	}
 
@@ -237,13 +241,12 @@ namespace HLODOutliner
 			/** Delegate to show all properties */
 			static bool IsPropertyVisible(const FPropertyAndParent& PropertyAndParent, bool bInShouldShowNonEditable)
 			{
-				if (PropertyAndParent.Property.GetName() == "HierarchicalLODSetup" || (PropertyAndParent.ParentProperty && PropertyAndParent.ParentProperty->GetName() == "MergeSetting") || (PropertyAndParent.ParentProperty && PropertyAndParent.ParentProperty->GetName() == "ProxySetting"))
+				if (PropertyAndParent.Property.GetName() == "HierarchicalLODSetup" || (PropertyAndParent.ParentProperty && PropertyAndParent.ParentProperty->GetName() == "MergeSetting") || (PropertyAndParent.ParentProperty && PropertyAndParent.ParentProperty->GetName() == "ProxySetting") || (PropertyAndParent.ParentProperty && PropertyAndParent.ParentProperty->GetName() == "Material"))
 				{
 					return true;
 				}
 				return false;
 			}
-
 		};
 
 		SettingsView->SetIsPropertyVisibleDelegate(FIsPropertyVisible::CreateStatic(&Local::IsPropertyVisible, true));
@@ -297,6 +300,18 @@ namespace HLODOutliner
 			TreeView->RequestTreeRefresh();
 		}
 
+		// Update the forced LOD level, as the slider for it is being dragged
+		if (bForcedSliderValueUpdating)
+		{
+			// Snap values
+			int32 SnappedValue = FMath::RoundToInt(FMath::Min(ForcedLODSliderValue, 1.0f) * (float)LODLevelDrawDistances.Num());
+			if (SnappedValue - 1 != ForcedLODLevel)
+			{
+				RestoreForcedLODLevel(ForcedLODLevel);
+				ForcedLODLevel = -1;
+				SetForcedLODLevel(SnappedValue - 1);
+			}
+		}
 	}
 
 	void SHLODOutliner::OnMouseEnter(const FGeometry& MyGeometry, const FPointerEvent& MouseEvent)
@@ -494,6 +509,8 @@ namespace HLODOutliner
 			RestoreForcedLODLevel(ForcedLODLevel);
 			SetForcedLODLevel(LevelItem->LODLevelIndex);			
 		}
+
+		ForcedLODSliderValue = ((1.0f / (LODLevelDrawDistances.Num())) * (ForcedLODLevel + 1));
 	}
 
 	bool SHLODOutliner::HandleForcedLevelSliderIsEnabled() const
@@ -509,19 +526,23 @@ namespace HLODOutliner
 
 	void SHLODOutliner::HandleForcedLevelSliderValueChanged(float NewValue)
 	{
-		// Snap values
-		int32 SnappedValue = FMath::Min( NewValue, 1.0f ) * LODLevelDrawDistances.Num();
-		if (SnappedValue - 1 != ForcedLODLevel)
-		{
-			RestoreForcedLODLevel(ForcedLODLevel);
-			ForcedLODLevel = -1;
-			SetForcedLODLevel(SnappedValue - 1);
-		}			
+		ForcedLODSliderValue = NewValue;					
+	}
+
+	void SHLODOutliner::HandleForcedLevelSliderCaptureBegin()
+	{
+		bForcedSliderValueUpdating = true;
+	}
+
+	void SHLODOutliner::HandleForcedLevelSliderCaptureEnd()
+	{	
+		bForcedSliderValueUpdating = false;		
+		ForcedLODSliderValue = ((1.0f / (LODLevelDrawDistances.Num())) * (ForcedLODLevel + 1));
 	}
 
 	float SHLODOutliner::HandleForcedLevelSliderValue() const
 	{
-		return ((1.0f / (LODLevelDrawDistances.Num())) * (ForcedLODLevel + 1));
+		return ForcedLODSliderValue;
 	}
 
 	FText SHLODOutliner::HandleForceLevelText() const
@@ -622,6 +643,25 @@ namespace HLODOutliner
 		}
 	}
 
+	void SHLODOutliner::RebuildLODActor(TSharedRef<ITreeItem> Item)
+	{
+		if (CurrentWorld)
+		{
+			FLODActorItem* ActorItem = (FLODActorItem*)(&Item.Get());
+
+			auto Parent = ActorItem->GetParent();
+
+			ITreeItem::TreeItemType Type = Parent->GetTreeItemType();
+			if (Type == ITreeItem::HierarchicalLODLevel)
+			{
+				FLODLevelItem* LevelItem = (FLODLevelItem*)(Parent.Get());
+				ActorItem->LODActor->SetIsDirty(true);
+				CurrentWorld->HierarchicalLODBuilder->BuildMeshForLODActor(ActorItem->LODActor.Get(), LevelItem->LODLevelIndex);
+				FullRefresh();
+			}
+		}
+	}
+
 	void SHLODOutliner::SelectLODActor(TSharedRef<ITreeItem> Item)
 	{
 		if (CurrentWorld)
@@ -857,7 +897,32 @@ namespace HLODOutliner
 
 	void SHLODOutliner::OnOutlinerDoubleClick(FTreeItemPtr TreeItem)
 	{
+		ITreeItem::TreeItemType Type = TreeItem->GetTreeItemType();
+		const bool bActiveViewportOnly = false;
 		
+		switch (Type)
+		{
+			case ITreeItem::HierarchicalLODLevel:
+			{
+				break;
+			}
+
+			case ITreeItem::HierarchicalLODActor:
+			{
+				FLODActorItem* ActorItem = (FLODActorItem*)(TreeItem.Get());
+				SelectActorInViewport(ActorItem->LODActor.Get(), 0);
+				GEditor->MoveViewportCamerasToActor(*ActorItem->LODActor.Get(), bActiveViewportOnly);
+				break;
+			}
+
+			case ITreeItem::StaticMeshActor:
+			{
+				FStaticMeshActorItem* StaticMeshActorItem = (FStaticMeshActorItem*)(TreeItem.Get());
+				SelectActorInViewport(StaticMeshActorItem->StaticMeshActor.Get(), 0);
+				GEditor->MoveViewportCamerasToActor(*StaticMeshActorItem->StaticMeshActor.Get(), bActiveViewportOnly);
+				break;
+			}
+		}	
 	}
 
 	TSharedPtr<SWidget> SHLODOutliner::OnOpenContextMenu()
@@ -874,7 +939,7 @@ namespace HLODOutliner
 		FMenuBuilder MenuBuilder(bCloseAfterSelection, TSharedPtr<FUICommandList>(), Extender);
 
 		const auto NumSelectedItems = TreeView->GetNumItemsSelected();
-		if (NumSelectedItems == 1)
+		if (NumSelectedItems == 1 && TreeView->GetSelectedItems()[0]->GetTreeItemType() != ITreeItem::HierarchicalLODLevel)
 		{
 			TreeView->GetSelectedItems()[0]->GenerateContextMenu(MenuBuilder, *this);
 
@@ -1211,42 +1276,47 @@ namespace HLODOutliner
 					TreeItemsMap.Add(LevelItem->GetID(), LevelItem);
 				}
 
-				for (AActor* Actor : CurrentWorld->GetCurrentLevel()->Actors)
+				for (ULevel* Level : CurrentWorld->GetLevels())
 				{
-					if (Actor && Actor->IsA<ALODActor>())
+					for (AActor* Actor : Level->Actors)
 					{
-						ALODActor* LODActor = CastChecked<ALODActor>(Actor);
-
-						if (LODActor)
+						if (Actor && Actor->IsA<ALODActor>())
 						{
-							FTreeItemRef Item = MakeShareable(new FLODActorItem(LODActor));
-							AllNodes.Add(Item->AsShared());
-
-							PendingActions.Emplace(FOutlinerAction::AddItem, Item, LevelNodes[LODActor->LODLevel - 1]);
+							ALODActor* LODActor = CastChecked<ALODActor>(Actor);
 
-							for (AActor* ChildActor : LODActor->SubActors)
+							if (LODActor)
 							{
-								if (ChildActor->IsA<ALODActor>())
+								FTreeItemRef Item = MakeShareable(new FLODActorItem(LODActor));
+								AllNodes.Add(Item->AsShared());
+
+								PendingActions.Emplace(FOutlinerAction::AddItem, Item, LevelNodes[LODActor->LODLevel - 1]);
+
+								for (AActor* ChildActor : LODActor->SubActors)
 								{
-									FTreeItemRef ChildItem = MakeShareable(new FLODActorItem(CastChecked<ALODActor>(ChildActor)));
-									AllNodes.Add(ChildItem->AsShared());
-									Item->AddChild(ChildItem);
+									if (ChildActor->IsA<ALODActor>())
+									{
+										FTreeItemRef ChildItem = MakeShareable(new FLODActorItem(CastChecked<ALODActor>(ChildActor)));
+										AllNodes.Add(ChildItem->AsShared());
+										Item->AddChild(ChildItem);
+									}
+									else
+									{
+										FTreeItemRef ChildItem = MakeShareable(new FStaticMeshActorItem(ChildActor));
+										AllNodes.Add(ChildItem->AsShared());
+
+										PendingActions.Emplace(FOutlinerAction::AddItem, ChildItem, Item);
+									}
 								}
-								else
-								{
-									FTreeItemRef ChildItem = MakeShareable(new FStaticMeshActorItem(ChildActor));
-									AllNodes.Add(ChildItem->AsShared());
 
-									PendingActions.Emplace(FOutlinerAction::AddItem, ChildItem, Item);
-								}
+								LODLevelBuildFlags[LODActor->LODLevel - 1] &= !LODActor->IsDirty();
+								LODLevelActors[LODActor->LODLevel - 1].Add(LODActor);
 							}
-
-							LODLevelBuildFlags[LODActor->LODLevel - 1] &= !LODActor->IsDirty();
-							LODLevelActors[LODActor->LODLevel - 1].Add(LODActor);
 						}
 					}
 				}
 
+				
+
 				for (uint32 LODLevelIndex = 0; LODLevelIndex < LODLevels; ++LODLevelIndex)
 				{
 					if (LODLevelActors[LODLevelIndex].Num() == 0)
diff --git a/Engine/Source/Editor/HierarchicalLODOutliner/Private/HLODOutliner.h b/Engine/Source/Editor/HierarchicalLODOutliner/Private/HLODOutliner.h
--- a/Engine/Source/Editor/HierarchicalLODOutliner/Private/HLODOutliner.h
+++ b/Engine/Source/Editor/HierarchicalLODOutliner/Private/HLODOutliner.h
@@ -147,10 +147,16 @@ namespace HLODOutliner
 		/**
 		* Handles changes in slider value, maps/snaps it to the corresponding HLOD level
 		*
-		* @param NewValue -
+		* @param NewValue - New value of the slider (0.0 - 1.0)
 		*/
 		void HandleForcedLevelSliderValueChanged(float NewValue);
 
+		/** Delegate for starting the forced HLOD level capture from the slider */
+		void HandleForcedLevelSliderCaptureBegin();
+		
+		/** Delegate for ending the forced HLOD level capture from the slider */
+		void HandleForcedLevelSliderCaptureEnd();
+
 		/**
 		* Returns the float sliding value corresponding to the currently forced HLOD level
 		*
@@ -201,6 +207,13 @@ namespace HLODOutliner
 		*/
 		void BuildLODActor(TSharedRef<ITreeItem> Item);	
 
+		/**
+		* Rebuilds the HLOD mesh for the given ALODActor (cluster)
+		*
+		* @param Item - (LODActor)TreeItem containing the ALODActor ptr
+		*/
+		void RebuildLODActor(TSharedRef<ITreeItem> Item);
+
 		/**
 		* Select the LODActor in the Editor Viewport
 		*
@@ -479,6 +492,11 @@ namespace HLODOutliner
 
 		/** Currently forced LOD level*/
 		int32 ForcedLODLevel;
+
+		/** Updating slider value for the Forced LOD level */
+		float ForcedLODSliderValue;
+		bool bForcedSliderValueUpdating;
+
 		/** Array with flags for each LOD level (whether or not all their Clusters/LODActors have their meshes built) */
 		TArray<bool> LODLevelBuildFlags;
 		/** Array of LODActors/Cluster per LOD level*/
diff --git a/Engine/Source/Editor/HierarchicalLODOutliner/Private/LODActorItem.cpp b/Engine/Source/Editor/HierarchicalLODOutliner/Private/LODActorItem.cpp
--- a/Engine/Source/Editor/HierarchicalLODOutliner/Private/LODActorItem.cpp
+++ b/Engine/Source/Editor/HierarchicalLODOutliner/Private/LODActorItem.cpp
@@ -38,6 +38,8 @@ void HLODOutliner::FLODActorItem::GenerateContextMenu(FMenuBuilder& MenuBuilder,
 	else
 	{		
 		MenuBuilder.AddMenuEntry(LOCTEXT("ForceView", "ForceView"), FText(), FSlateIcon(), FUIAction(FExecuteAction::CreateSP(&Outliner, &SHLODOutliner::ForceViewLODActor, AsShared())));
+
+		MenuBuilder.AddMenuEntry(LOCTEXT("BuildLODActorMesh", "Rebuild LOD Mesh"), FText(), FSlateIcon(), FUIAction(FExecuteAction::CreateSP(&Outliner, &SHLODOutliner::RebuildLODActor, AsShared())));
 	}
 
 	MenuBuilder.AddMenuEntry(LOCTEXT("CreateHLODVolume", "Create Containing Hierarchical Volume"), FText(), FSlateIcon(), FUIAction(FExecuteAction::CreateSP(&Outliner, &SHLODOutliner::CreateHierarchicalVolumeForActor, AsShared())));
diff --git a/Engine/Source/Editor/LevelEditor/Private/LevelEditor.cpp b/Engine/Source/Editor/LevelEditor/Private/LevelEditor.cpp
--- a/Engine/Source/Editor/LevelEditor/Private/LevelEditor.cpp
+++ b/Engine/Source/Editor/LevelEditor/Private/LevelEditor.cpp
@@ -1377,10 +1377,7 @@ void FLevelEditorModule::BindGlobalLevelEditorCommands()
 
 	ActionList.MapAction(Commands.BuildLODsOnly,
 		FExecuteAction::CreateStatic(&FLevelEditorActionCallbacks::BuildLODsOnly_Execute));
-
-	ActionList.MapAction(Commands.PreviewHLODClustersOnly,
-		FExecuteAction::CreateStatic(&FLevelEditorActionCallbacks::PreviewHLODClustersOnly_Execute));
-
+	
 	ActionList.MapAction( 
 		Commands.LightingQuality_Production, 
 		FExecuteAction::CreateStatic( &FLevelEditorActionCallbacks::SetLightingQuality, (ELightingBuildQuality)Quality_Production ),
diff --git a/Engine/Source/Editor/LevelEditor/Private/LevelEditorActions.cpp b/Engine/Source/Editor/LevelEditor/Private/LevelEditorActions.cpp
--- a/Engine/Source/Editor/LevelEditor/Private/LevelEditorActions.cpp
+++ b/Engine/Source/Editor/LevelEditor/Private/LevelEditorActions.cpp
@@ -633,12 +633,6 @@ void FLevelEditorActionCallbacks::BuildLODsOnly_Execute()
 	FEditorBuildUtils::EditorBuild(GetWorld(), EBuildOptions::BuildHierarchicalLOD);
 }
 
-void FLevelEditorActionCallbacks::PreviewHLODClustersOnly_Execute()
-{
-	// Preview HLOD Clusters
-	FEditorBuildUtils::EditorBuild(GetWorld(), EBuildOptions::PreviewHierarchicalLOD);
-}
-
 bool FLevelEditorActionCallbacks::IsLightingQualityChecked( ELightingBuildQuality TestQuality )
 {
 	int32 CurrentQualityLevel;
@@ -2834,7 +2828,7 @@ void FLevelEditorCommands::RegisterCommands()
 	UI_COMMAND( BuildGeometryOnly_OnlyCurrentLevel, "Build Geometry (Current Level)", "Builds geometry, only for the current level", EUserInterfaceActionType::Button, FInputChord() );
 	UI_COMMAND( BuildPathsOnly, "Build Paths", "Only builds paths (all levels.)", EUserInterfaceActionType::Button, FInputChord() );
 	UI_COMMAND( BuildLODsOnly, "Build LODs", "Only builds LODs (all levels.)", EUserInterfaceActionType::Button, FInputChord() );
-	UI_COMMAND( PreviewHLODClustersOnly, "Preview HLOD clusters", "Preview builds LODs and shows clusters (all levels.)", EUserInterfaceActionType::Button, FInputChord());
+	
 	UI_COMMAND( LightingQuality_Production, "Production", "Sets precomputed lighting quality to highest possible quality (slowest computation time.)", EUserInterfaceActionType::RadioButton, FInputChord() );
 	UI_COMMAND( LightingQuality_High, "High", "Sets precomputed lighting quality to high quality", EUserInterfaceActionType::RadioButton, FInputChord() );
 	UI_COMMAND( LightingQuality_Medium, "Medium", "Sets precomputed lighting quality to medium quality", EUserInterfaceActionType::RadioButton, FInputChord() );
diff --git a/Engine/Source/Editor/LevelEditor/Private/LevelEditorToolBar.cpp b/Engine/Source/Editor/LevelEditor/Private/LevelEditorToolBar.cpp
--- a/Engine/Source/Editor/LevelEditor/Private/LevelEditorToolBar.cpp
+++ b/Engine/Source/Editor/LevelEditor/Private/LevelEditorToolBar.cpp
@@ -1565,7 +1565,6 @@ TSharedRef< SWidget > FLevelEditorToolBar::GenerateBuildMenuContent( TSharedRef<
 	MenuBuilder.BeginSection("LevelEditorLOD", LOCTEXT("LODHeading", "Hierarchical LOD"));
 	{
 		MenuBuilder.AddMenuEntry(FLevelEditorCommands::Get().BuildLODsOnly);
-		MenuBuilder.AddMenuEntry(FLevelEditorCommands::Get().PreviewHLODClustersOnly);
 	}
 	MenuBuilder.EndSection();
 
diff --git a/Engine/Source/Editor/LevelEditor/Public/LevelEditorActions.h b/Engine/Source/Editor/LevelEditor/Public/LevelEditorActions.h
--- a/Engine/Source/Editor/LevelEditor/Public/LevelEditorActions.h
+++ b/Engine/Source/Editor/LevelEditor/Public/LevelEditorActions.h
@@ -71,7 +71,6 @@ public:
 	TSharedPtr< FUICommandInfo > BuildGeometryOnly_OnlyCurrentLevel;
 	TSharedPtr< FUICommandInfo > BuildPathsOnly;
 	TSharedPtr< FUICommandInfo > BuildLODsOnly;
-	TSharedPtr< FUICommandInfo > PreviewHLODClustersOnly;
 	TSharedPtr< FUICommandInfo > LightingQuality_Production;
 	TSharedPtr< FUICommandInfo > LightingQuality_High;
 	TSharedPtr< FUICommandInfo > LightingQuality_Medium;
@@ -702,7 +701,6 @@ public:
 	static void BuildGeometryOnly_OnlyCurrentLevel_Execute();
 	static void BuildPathsOnly_Execute();
 	static void BuildLODsOnly_Execute();
-	static void PreviewHLODClustersOnly_Execute();
 	static void SetLightingQuality( ELightingBuildQuality NewQuality );
 	static bool IsLightingQualityChecked( ELightingBuildQuality TestQuality );
 	static float GetLightingDensityIdeal();
diff --git a/Engine/Source/Editor/UnrealEd/Private/EditorBuildUtils.cpp b/Engine/Source/Editor/UnrealEd/Private/EditorBuildUtils.cpp
--- a/Engine/Source/Editor/UnrealEd/Private/EditorBuildUtils.cpp
+++ b/Engine/Source/Editor/UnrealEd/Private/EditorBuildUtils.cpp
@@ -217,7 +217,6 @@ bool FEditorBuildUtils::EditorBuild( UWorld* InWorld, EBuildOptions::Type Id, co
 		BuildType = SBuildProgressWidget::BUILDTYPE_Paths;
 		break;
 	case EBuildOptions::BuildHierarchicalLOD:
-	case EBuildOptions::PreviewHierarchicalLOD:
 		BuildType = SBuildProgressWidget::BUILDTYPE_LODs;
 		break;
 	default:
@@ -307,24 +306,7 @@ bool FEditorBuildUtils::EditorBuild( UWorld* InWorld, EBuildOptions::Type Id, co
 			bDoBuild = GEditor->WarnAboutHiddenLevels( InWorld, false );
 			if ( bDoBuild )
 			{
-				GEditor->ResetTransaction( NSLOCTEXT("UnrealEd", "RebuildLOD", "Rebuilding HierarchicalLOD") );
-
-				// We can't set the busy cursor for all windows, because lighting
-				// needs a cursor for the lighting options dialog.
-				const FScopedBusyCursor BusyCursor;
-
-				TriggerHierarchicalLODBuilder(InWorld, Id);
-			}
-
-			break;
-		}
-
-	case  EBuildOptions::PreviewHierarchicalLOD:
-		{
-			bDoBuild = GEditor->WarnAboutHiddenLevels(InWorld, false);
-			if (bDoBuild)
-			{
-				GEditor->ResetTransaction(NSLOCTEXT("UnrealEd", "RebuildLOD", "Rebuilding HierarchicalLOD"));
+				GEditor->ResetTransaction( NSLOCTEXT("UnrealEd", "BuildHLODMeshes", "Building Hierarchical LOD Meshes") );
 
 				// We can't set the busy cursor for all windows, because lighting
 				// needs a cursor for the lighting options dialog.
@@ -818,14 +800,7 @@ void FEditorBuildUtils::TriggerNavigationBuilder(UWorld* InWorld, EBuildOptions:
 void FEditorBuildUtils::TriggerHierarchicalLODBuilder(UWorld* InWorld, EBuildOptions::Type Id)
 {
 	// Invoke HLOD generator, with either preview or full build
-	if (Id == EBuildOptions::PreviewHierarchicalLOD)
-	{
-		InWorld->HierarchicalLODBuilder->PreviewBuild();
-	}
-	else
-	{
-		InWorld->HierarchicalLODBuilder->Build();
-	}
+	InWorld->HierarchicalLODBuilder->BuildMeshesForLODActors();
 }
 
 #undef LOCTEXT_NAMESPACE
\ No newline at end of file
diff --git a/Engine/Source/Editor/UnrealEd/Public/EditorBuildUtils.h b/Engine/Source/Editor/UnrealEd/Public/EditorBuildUtils.h
--- a/Engine/Source/Editor/UnrealEd/Public/EditorBuildUtils.h
+++ b/Engine/Source/Editor/UnrealEd/Public/EditorBuildUtils.h
@@ -28,9 +28,7 @@ namespace EBuildOptions
 		/** Build everything except for paths only build seleced */
 		BuildAllOnlySelectedPaths,
 		/** Build Hierarchical LOD system - need WorldSetting setup*/
-		BuildHierarchicalLOD, 
-		/** Preview Hierarchical LOD system clusters - need WorldSetting setup*/
-		PreviewHierarchicalLOD,
+		BuildHierarchicalLOD
 	};
 }
 
diff --git a/Engine/Source/Runtime/Core/Private/HAL/ConsoleManager.cpp b/Engine/Source/Runtime/Core/Private/HAL/ConsoleManager.cpp
--- a/Engine/Source/Runtime/Core/Private/HAL/ConsoleManager.cpp
+++ b/Engine/Source/Runtime/Core/Private/HAL/ConsoleManager.cpp
@@ -2110,4 +2110,10 @@ static TAutoConsoleVariable<int32> CVarCheckSRVTransitions(
 	TEXT("r.CheckSRVTransitions"),
 	0,
 	TEXT("Tests that render targets are properly transitioned to SRV when SRVs are set."),
-	ECVF_RenderThreadSafe);  
\ No newline at end of file
+	ECVF_RenderThreadSafe);  
+
+static TAutoConsoleVariable<int32> CVarHLODSystemEnabled(
+	TEXT("r.HLODEnabled"), 
+	1,
+	TEXT("Toggles whether or not the Hierarchical LOD system is enabled."),
+	ECVF_Scalability | ECVF_RenderThreadSafe);
\ No newline at end of file
diff --git a/Engine/Source/Runtime/Engine/Classes/Engine/LODActor.h b/Engine/Source/Runtime/Engine/Classes/Engine/LODActor.h
--- a/Engine/Source/Runtime/Engine/Classes/Engine/LODActor.h
+++ b/Engine/Source/Runtime/Engine/Classes/Engine/LODActor.h
@@ -17,7 +17,7 @@ GENERATED_UCLASS_BODY()
 
 private_subobject:
 	// disable display of this component
-	UPROPERTY(/*Category=LODActor, VisibleAnywhere*/)
+	UPROPERTY(Category=LODActor, VisibleAnywhere)
 	class UStaticMeshComponent* StaticMeshComponent;
 
 public:
@@ -92,6 +92,9 @@ public:
 	
 	/** Updates the LODParents for the SubActors (and the drawing distance)*/
 	void UpdateSubActorLODParents();
+
+	/** Cleans the SubActor array (removes all NULL entries) */
+	void CleanSubActorArray();
 #endif // WITH_EDITOR
 	
 	// Begin UObject interface.
diff --git a/Engine/Source/Runtime/Engine/Private/LODActor.cpp b/Engine/Source/Runtime/Engine/Private/LODActor.cpp
--- a/Engine/Source/Runtime/Engine/Private/LODActor.cpp
+++ b/Engine/Source/Runtime/Engine/Private/LODActor.cpp
@@ -14,6 +14,33 @@
 
 #define LOCTEXT_NAMESPACE "LODActor"
 
+static void ToggleHLODEnabled(UWorld* InWorld)
+{
+	FlushRenderingCommands();
+	
+	ULevel* CurrentLevel = InWorld->GetCurrentLevel();	
+	
+	bool bHLODEnabled = ( IConsoleManager::Get().FindConsoleVariable(TEXT("r.HLODEnabled"))->GetInt() == 1 );
+	IConsoleManager::Get().FindConsoleVariable(TEXT("r.HLODEnabled"))->Set((int32)!bHLODEnabled);
+
+	for (AActor* Actor : CurrentLevel->Actors)
+	{
+		ALODActor* LODActor = Cast<ALODActor>(Actor);
+		if (LODActor)
+		{
+			LODActor->SetHidden(bHLODEnabled);
+			LODActor->SetIsTemporarilyHiddenInEditor(bHLODEnabled);
+			LODActor->MarkComponentsRenderStateDirty();			
+		}
+	}
+}
+
+static FAutoConsoleCommandWithWorld GToggleHLODEnabledCmd(
+	TEXT("ToggleHLODEnabled"),
+	TEXT("Toggles whether or not the HLOD system is enabled."),
+	FConsoleCommandWithWorldDelegate::CreateStatic(ToggleHLODEnabled)
+	);
+
 ALODActor::ALODActor(const FObjectInitializer& ObjectInitializer)
 	: Super(ObjectInitializer)
 	, LODDrawDistance(5000)	
@@ -56,13 +83,13 @@ void ALODActor::PostRegisterAllComponents()
 		(StaticMeshComponent->SceneProxy)->SetHierarchicalLOD_GameThread(LODLevel);
 	}
 
-	UpdateSubActorLODParents();
+	// Clean up sub actor if assets were delete manually
+	CleanSubActorArray();
 
+	UpdateSubActorLODParents();	
 #endif
 }
 
-
-
 #if WITH_EDITOR
 
 void ALODActor::PostEditChangeProperty(FPropertyChangedEvent& PropertyChangedEvent)
@@ -230,7 +257,7 @@ void ALODActor::RemoveSubActor(AActor* InActor)
 	StaticMeshComponent->MarkRenderStateDirty();
 
 	// In case the user removes an actor while the HLOD system is force viewing one LOD level
-	InActor->SetIsTemporarilyHiddenInEditor(false);
+	InActor->SetIsTemporarilyHiddenInEditor(false);	
 }
 
 void ALODActor::SetIsDirty(const bool bNewState)
@@ -337,7 +364,6 @@ const uint32 ALODActor::GetNumTriangles()
 	return NumTrianglesInSubActors;
 }
 
-
 void ALODActor::SetStaticMesh(class UStaticMesh* InStaticMesh)
 {
 	if (StaticMeshComponent)
@@ -355,6 +381,19 @@ void ALODActor::UpdateSubActorLODParents()
 	}
 }
 
+void ALODActor::CleanSubActorArray()
+{
+	for (int32 SubActorIndex = 0; SubActorIndex < SubActors.Num(); ++SubActorIndex)
+	{
+		auto& Actor = SubActors[SubActorIndex];
+		if (Actor == nullptr)
+		{
+			SubActors.RemoveAt(SubActorIndex);
+			SubActorIndex--;
+		}
+	}
+}
+
 #endif // WITH_EDITOR
 
 FBox ALODActor::GetComponentsBoundingBox(bool bNonColliding) const 
diff --git a/Engine/Source/Runtime/Renderer/Private/ScenePrivate.h b/Engine/Source/Runtime/Renderer/Private/ScenePrivate.h
--- a/Engine/Source/Runtime/Renderer/Private/ScenePrivate.h
+++ b/Engine/Source/Runtime/Renderer/Private/ScenePrivate.h
@@ -1395,7 +1395,7 @@ public:
 	void UpdateNodeSceneInfo(FPrimitiveComponentId NodeId, FPrimitiveSceneInfo* SceneInfo);
 	void PopulateHiddenFlags(FViewInfo& View, FSceneBitArray& HiddenFlags);
 
-	bool IsActive() { return SceneNodes.Num() > 0; }
+	bool IsActive() { return (SceneNodes.Num() > 0) && (IConsoleManager::Get().FindConsoleVariable(TEXT("r.HLODEnabled"))->GetInt() == 1); }
 
 private:
 	/** Scene this Tree belong to */
