UE-19601 HLOD : Set up better default setting whenever adding new element
Workflow fixes (lodactor and lodlevel Tree Items are expanded by default now)

diff --git a/Engine/Source/Editor/HierarchicalLODOutliner/HierarchicalLODOutliner.Build.cs b/Engine/Source/Editor/HierarchicalLODOutliner/HierarchicalLODOutliner.Build.cs
--- a/Engine/Source/Editor/HierarchicalLODOutliner/HierarchicalLODOutliner.Build.cs
+++ b/Engine/Source/Editor/HierarchicalLODOutliner/HierarchicalLODOutliner.Build.cs
@@ -6,7 +6,7 @@ public class HierarchicalLODOutliner : ModuleRules
 {
     public HierarchicalLODOutliner(TargetInfo Target)
 	{
-		PublicDependencyModuleNames.AddRange
+        PublicDependencyModuleNames.AddRange
         (
             new string[] {
 				"Core",
@@ -21,26 +21,29 @@ public class HierarchicalLODOutliner : ModuleRules
                 "PropertyEditor"
 			}
         );
-
-
+        
         PrivateDependencyModuleNames.AddRange(
              new string[] {
 					"Engine",
                     "UnrealEd"
 				}
          );
-        
-        PrivateDependencyModuleNames.AddRange(
-            new string[] {
-				"SlateCore",
-			}
-        );
 
-        PrivateIncludePathModuleNames.AddRange(
-        new string[] {
-				"Messaging",
-				"SessionServices",
-			}
-    );
+        if (UEBuildConfiguration.bBuildEditor == true)
+        {
+            PrivateDependencyModuleNames.AddRange(
+                new string[] {
+                    "UnrealEd",
+    				"SlateCore",
+    				"Slate",
+                }
+            );
+
+            CircularlyReferencedDependentModules.AddRange(
+                new string[] {
+                    "UnrealEd"
+                }
+            );
+        }
 	}
 }
diff --git a/Engine/Source/Editor/HierarchicalLODOutliner/Private/HLODOutliner.cpp b/Engine/Source/Editor/HierarchicalLODOutliner/Private/HLODOutliner.cpp
--- a/Engine/Source/Editor/HierarchicalLODOutliner/Private/HLODOutliner.cpp
+++ b/Engine/Source/Editor/HierarchicalLODOutliner/Private/HLODOutliner.cpp
@@ -255,14 +255,14 @@ namespace HLODOutliner
 
 	void SHLODOutliner::Tick(const FGeometry& AllottedGeometry, const double InCurrentTime, const float InDeltaTime)
 	{
-		// Get a collection of items and folders which were formerly collapsed
-		const FParentsExpansionState ExpansionStateInfo = GetParentsExpansionState();
-
 		if (bNeedsRefresh)
 		{
 			Populate();
 		}
 
+		// Get a collection of items and folders which were formerly collapsed
+		const FParentsExpansionState ExpansionStateInfo = GetParentsExpansionState();
+
 		bool bChangeMade = false;
 
 		// Only deal with 256 at a time
@@ -292,13 +292,13 @@ namespace HLODOutliner
 		}
 		PendingActions.RemoveAt(0, End);
 
-		// Restore expansion states
-		SetParentsExpansionState(ExpansionStateInfo);
-
 		if (bChangeMade)
 		{
+			// Restore expansion states
+			SetParentsExpansionState(ExpansionStateInfo);
+
 			TreeView->RequestTreeRefresh();
-		}
+		}			
 
 		// Update the forced LOD level, as the slider for it is being dragged
 		if (bForcedSliderValueUpdating)
@@ -467,6 +467,7 @@ namespace HLODOutliner
 		GEngine->OnHLODActorMoved().RemoveAll(this);
 		GEngine->OnHLODActorAdded().RemoveAll(this);
 		GEngine->OnHLODActorMarkedDirty().RemoveAll(this);
+		GEngine->OnHLODLevelsArrayChanged().RemoveAll(this);
 	}
 
 	void SHLODOutliner::ForceViewLODActor(TSharedRef<ITreeItem> Item)
@@ -1196,28 +1197,6 @@ namespace HLODOutliner
 	{
 		if (CurrentWorld)
 		{
-			auto WorldSettings = CurrentWorld->GetWorldSettings();			
-			if (WorldSettings->HierarchicalLODSetup.Num() > 1 && WorldSettings->HierarchicalLODSetup.Num() > LODLevelDrawDistances.Num())
-			{
-				// HLOD level was added, use previous level settings to "guess" new settings
-				auto& NewLevelSetup = WorldSettings->HierarchicalLODSetup.Last();
-				auto& OldLastLevelSetup = WorldSettings->HierarchicalLODSetup[WorldSettings->HierarchicalLODSetup.Num() - 2];
-
-				NewLevelSetup.bSimplifyMesh = OldLastLevelSetup.bSimplifyMesh;
-				NewLevelSetup.MergeSetting = OldLastLevelSetup.MergeSetting;
-				NewLevelSetup.ProxySetting = OldLastLevelSetup.ProxySetting;
-
-				NewLevelSetup.DesiredBoundRadius = OldLastLevelSetup.DesiredBoundRadius * 2.5f;
-				NewLevelSetup.DesiredFillingPercentage = FMath::Max( OldLastLevelSetup.DesiredFillingPercentage * 0.75f, 1.0f );
-				NewLevelSetup.DrawDistance = OldLastLevelSetup.DrawDistance * 2.5f;
-				NewLevelSetup.MinNumberOfActorsToBuild = OldLastLevelSetup.MinNumberOfActorsToBuild;			
-			}
-			else if (WorldSettings->HierarchicalLODSetup.Num() < LODLevelDrawDistances.Num())
-			{
-				// HLOD Level was removed, now remove all LODActors for this level
-				RemoveLODLevelActors(LODLevelDrawDistances.Num() - 1);
-			}
-
 			FullRefresh();
 		}
 	}
@@ -1274,6 +1253,9 @@ namespace HLODOutliner
 					LODLevelDrawDistances.Add(WorldSettings->HierarchicalLODSetup[LODLevelIndex].DrawDistance);
 
 					TreeItemsMap.Add(LevelItem->GetID(), LevelItem);
+
+					// Expand level items by default
+					LevelItem->bIsExpanded = true;
 				}
 
 				for (ULevel* Level : CurrentWorld->GetLevels())
diff --git a/Engine/Source/Editor/HierarchicalLODOutliner/Private/HierarchicalLODOutlinerModule.cpp b/Engine/Source/Editor/HierarchicalLODOutliner/Private/HierarchicalLODOutlinerModule.cpp
--- a/Engine/Source/Editor/HierarchicalLODOutliner/Private/HierarchicalLODOutlinerModule.cpp
+++ b/Engine/Source/Editor/HierarchicalLODOutliner/Private/HierarchicalLODOutlinerModule.cpp
@@ -3,15 +3,66 @@
 #include "HierarchicalLODOutlinerPrivatePCH.h"
 #include "HierarchicalLODOutlinerModule.h"
 #include "HLODOutliner.h"
+#include "HierarchicalLODUtils.h"
 
 void FHierarchicalLODOutlinerModule::StartupModule()
 {
-
+	GUnrealEd->OnHLODLevelsArrayChanged().AddRaw(this, &FHierarchicalLODOutlinerModule::OnHLODLevelsArrayChangedEvent);
 }
 
 void FHierarchicalLODOutlinerModule::ShutdownModule()
 {
+	GUnrealEd->OnHLODLevelsArrayChanged().RemoveAll(this);
+}
+
+void FHierarchicalLODOutlinerModule::OnHLODLevelsArrayChangedEvent()
+{
+	UWorld* CurrentWorld = nullptr;
+	
+	// Get the correct UWorld instance
+	for (const FWorldContext& Context : GEngine->GetWorldContexts())
+	{
+		if (Context.WorldType == EWorldType::PIE)
+		{
+			CurrentWorld = Context.World();
+			break;
+		}
+		else if (Context.WorldType == EWorldType::Editor)
+		{
+			CurrentWorld = Context.World();
+		}
+	}
+
+	AWorldSettings* WorldSettings = (CurrentWorld) ? CurrentWorld->GetWorldSettings() : nullptr;
+
+	if (WorldSettings)
+	{
+		auto& HierarchicalLODSetup = WorldSettings->HierarchicalLODSetup;
+		auto& NumHLODLevels = WorldSettings->NumHLODLevels;
+
+		if (HierarchicalLODSetup.Num() > 1 && HierarchicalLODSetup.Num() > NumHLODLevels)
+		{
+			// HLOD level was added, use previous level settings to "guess" new settings
+			auto& NewLevelSetup = HierarchicalLODSetup.Last();
+			auto& OldLastLevelSetup = HierarchicalLODSetup[HierarchicalLODSetup.Num() - 2];
+
+			NewLevelSetup.bSimplifyMesh = OldLastLevelSetup.bSimplifyMesh;
+			NewLevelSetup.MergeSetting = OldLastLevelSetup.MergeSetting;
+			NewLevelSetup.ProxySetting = OldLastLevelSetup.ProxySetting;
+
+			NewLevelSetup.DesiredBoundRadius = OldLastLevelSetup.DesiredBoundRadius * 2.5f;
+			NewLevelSetup.DesiredFillingPercentage = FMath::Max(OldLastLevelSetup.DesiredFillingPercentage * 0.75f, 1.0f);
+			NewLevelSetup.DrawDistance = OldLastLevelSetup.DrawDistance * 2.5f;
+			NewLevelSetup.MinNumberOfActorsToBuild = OldLastLevelSetup.MinNumberOfActorsToBuild;
+		}
+		else if (HierarchicalLODSetup.Num() < NumHLODLevels)
+		{
+			// HLOD Level was removed, now remove all LODActors for this level
+			HierarchicalLODUtils::DeleteLODActorsInHLODLevel(CurrentWorld, NumHLODLevels - 1);
+		}
+	}
 
+	
 }
 
 TSharedRef<SWidget> FHierarchicalLODOutlinerModule::CreateHLODOutlinerWidget()
diff --git a/Engine/Source/Editor/HierarchicalLODOutliner/Private/LODActorItem.cpp b/Engine/Source/Editor/HierarchicalLODOutliner/Private/LODActorItem.cpp
--- a/Engine/Source/Editor/HierarchicalLODOutliner/Private/LODActorItem.cpp
+++ b/Engine/Source/Editor/HierarchicalLODOutliner/Private/LODActorItem.cpp
@@ -125,6 +125,9 @@ void HLODOutliner::FLODActorItem::OnDrop(FDragDropPayload& DraggedObjects, const
 {
 	FLODActorDropTarget Target(LODActor.Get());
 	Target.OnDrop(DraggedObjects, ValidationInfo, DroppedOnWidget);
+
+	// Expand this HLOD actor item
+	bIsExpanded = true;	
 }
 
 HLODOutliner::FDragValidationInfo HLODOutliner::FLODActorDropTarget::ValidateDrop(FDragDropPayload& DraggedObjects) const
@@ -248,7 +251,7 @@ void HLODOutliner::FLODActorDropTarget::OnDrop(FDragDropPayload& DraggedObjects,
 	{
 		return;
 	}
-	
+
 	auto& DraggedStaticMeshActors = DraggedObjects.StaticMeshActors.GetValue();
 	auto& DraggedLODActors = DraggedObjects.LODActors.GetValue();
 	if (ValidationInfo.TooltipType == FHLODOutlinerDragDropOp::ToolTip_CompatibleMoveToCluster || ValidationInfo.TooltipType == FHLODOutlinerDragDropOp::ToolTip_MultipleSelection_CompatibleMoveToCluster)
diff --git a/Engine/Source/Editor/HierarchicalLODOutliner/Private/LODLevelItem.cpp b/Engine/Source/Editor/HierarchicalLODOutliner/Private/LODLevelItem.cpp
--- a/Engine/Source/Editor/HierarchicalLODOutliner/Private/LODLevelItem.cpp
+++ b/Engine/Source/Editor/HierarchicalLODOutliner/Private/LODLevelItem.cpp
@@ -40,7 +40,7 @@ HLODOutliner::FTreeItemID HLODOutliner::FLODLevelItem::GetID()
 
 void HLODOutliner::FLODLevelItem::PopulateDragDropPayload(FDragDropPayload& Payload) const
 {
-
+	// Nothing to populate
 }
 
 HLODOutliner::FDragValidationInfo HLODOutliner::FLODLevelItem::ValidateDrop(FDragDropPayload& DraggedObjects) const
@@ -51,8 +51,11 @@ HLODOutliner::FDragValidationInfo HLODOutliner::FLODLevelItem::ValidateDrop(FDra
 
 void HLODOutliner::FLODLevelItem::OnDrop(FDragDropPayload& DraggedObjects, const FDragValidationInfo& ValidationInfo, TSharedRef<SWidget> DroppedOnWidget)
 {
-	FLODLevelDropTarget Target(LODLevelIndex);
-	return Target.OnDrop(DraggedObjects, ValidationInfo, DroppedOnWidget);
+	FLODLevelDropTarget Target(LODLevelIndex);		
+	Target.OnDrop(DraggedObjects, ValidationInfo, DroppedOnWidget);
+
+	// Expand this HLOD level item
+	bIsExpanded = true;
 }
 
 HLODOutliner::FDragValidationInfo HLODOutliner::FLODLevelDropTarget::ValidateDrop(FDragDropPayload& DraggedObjects) const
@@ -100,8 +103,8 @@ void HLODOutliner::FLODLevelDropTarget::OnDrop(FDragDropPayload& DraggedObjects,
 {
 	if (ValidationInfo.TooltipType == FHLODOutlinerDragDropOp::ToolTip_CompatibleNewCluster || ValidationInfo.TooltipType == FHLODOutlinerDragDropOp::ToolTip_MultipleSelection_CompatibleNewCluster)
 	{
-		CreateNewCluster(DraggedObjects);
-	}
+		CreateNewCluster(DraggedObjects);		
+	}	
 }
 
 void HLODOutliner::FLODLevelDropTarget::CreateNewCluster(FDragDropPayload &DraggedObjects)
diff --git a/Engine/Source/Editor/HierarchicalLODOutliner/Public/HierarchicalLODOutlinerModule.h b/Engine/Source/Editor/HierarchicalLODOutliner/Public/HierarchicalLODOutlinerModule.h
--- a/Engine/Source/Editor/HierarchicalLODOutliner/Public/HierarchicalLODOutlinerModule.h
+++ b/Engine/Source/Editor/HierarchicalLODOutliner/Public/HierarchicalLODOutlinerModule.h
@@ -19,6 +19,8 @@ public:
 	*/
 	virtual void ShutdownModule();
 
+	void OnHLODLevelsArrayChangedEvent();
+
 	/** Creates the HLOD Outliner widget */
 	virtual TSharedRef<SWidget> CreateHLODOutlinerWidget();
 };
\ No newline at end of file
diff --git a/Engine/Source/Editor/LevelEditor/LevelEditor.Build.cs b/Engine/Source/Editor/LevelEditor/LevelEditor.Build.cs
--- a/Engine/Source/Editor/LevelEditor/LevelEditor.Build.cs
+++ b/Engine/Source/Editor/LevelEditor/LevelEditor.Build.cs
@@ -17,7 +17,7 @@ public class LevelEditor : ModuleRules
 				"SlateReflector",
                 "IntroTutorials",
                 "AppFramework",
-                "PortalServices",
+                "PortalServices"
 			}
 		);
 
@@ -25,8 +25,7 @@ public class LevelEditor : ModuleRules
 			new string[] {
 				"Settings",
 				"UserFeedback",
-				"IntroTutorials",
-                "HierarchicalLODOutliner"
+				"IntroTutorials"
 			}
 		);
 
@@ -62,7 +61,8 @@ public class LevelEditor : ModuleRules
                 "Kismet",
 				"KismetWidgets",
 				"Sequencer",
-                "Foliage"
+                "Foliage",
+                "HierarchicalLODOutliner"
 			}
 		);
 
@@ -86,8 +86,7 @@ public class LevelEditor : ModuleRules
 				"UserFeedback",
 				"ReferenceViewer",
 				"SizeMap",
-                "IntroTutorials",
-                "HierarchicalLODOutliner"
+                "IntroTutorials"
 			}
 		);
 	}
diff --git a/Engine/Source/Editor/UnrealEd/Private/UnrealEd.cpp b/Engine/Source/Editor/UnrealEd/Private/UnrealEd.cpp
--- a/Engine/Source/Editor/UnrealEd/Private/UnrealEd.cpp
+++ b/Engine/Source/Editor/UnrealEd/Private/UnrealEd.cpp
@@ -154,6 +154,8 @@ int32 EditorInit( IEngineLoop& EngineLoop )
 		}
 	}
 
+	FModuleManager::LoadModuleChecked<IModuleInterface>(TEXT("HierarchicalLODOutliner"));
+
 	// this will be ultimately returned from main(), so no error should be 0.
 	return 0;
 }
diff --git a/Engine/Source/Runtime/Engine/Classes/GameFramework/WorldSettings.h b/Engine/Source/Runtime/Engine/Classes/GameFramework/WorldSettings.h
--- a/Engine/Source/Runtime/Engine/Classes/GameFramework/WorldSettings.h
+++ b/Engine/Source/Runtime/Engine/Classes/GameFramework/WorldSettings.h
@@ -420,6 +420,9 @@ class ENGINE_API AWorldSettings : public AInfo, public IInterface_AssetUserData
 	/** Hierarchical LOD Setup */
 	UPROPERTY(EditAnywhere, Category=LODSystem, meta=(editcondition = "bEnableHierarchicalLODSystem"))
 	TArray<struct FHierarchicalSimplification>	HierarchicalLODSetup;
+
+	UPROPERTY()
+	int32 NumHLODLevels;
 #endif
 	/************************************/
 	/** DEFAULT SETTINGS **/
diff --git a/Engine/Source/Runtime/Engine/Private/WorldSettings.cpp b/Engine/Source/Runtime/Engine/Private/WorldSettings.cpp
--- a/Engine/Source/Runtime/Engine/Private/WorldSettings.cpp
+++ b/Engine/Source/Runtime/Engine/Private/WorldSettings.cpp
@@ -43,6 +43,7 @@ AWorldSettings::AWorldSettings(const FObjectInitializer& ObjectInitializer)
 
  	FHierarchicalSimplification LODBaseSetup;
 	HierarchicalLODSetup.Add(LODBaseSetup);
+	NumHLODLevels = HierarchicalLODSetup.Num();
 #endif
 
 	KillZ = -HALF_WORLD_MAX1;
@@ -312,6 +313,7 @@ void AWorldSettings::PostEditChangeProperty(FPropertyChangedEvent& PropertyChang
 		if (PropertyThatChanged->GetName() == TEXT("HierarchicalLODSetup"))
 		{
 			GEditor->BroadcastHLODLevelsArrayChanged();
+			NumHLODLevels = HierarchicalLODSetup.Num();			
 		}
 	}
 
