Fix for crash when closing editor

diff --git a/Engine/Source/Editor/HierarchicalLODOutliner/Private/HierarchicalLODOutlinerModule.cpp b/Engine/Source/Editor/HierarchicalLODOutliner/Private/HierarchicalLODOutlinerModule.cpp
--- a/Engine/Source/Editor/HierarchicalLODOutliner/Private/HierarchicalLODOutlinerModule.cpp
+++ b/Engine/Source/Editor/HierarchicalLODOutliner/Private/HierarchicalLODOutlinerModule.cpp
@@ -7,12 +7,16 @@
 
 void FHierarchicalLODOutlinerModule::StartupModule()
 {
-	GUnrealEd->OnHLODLevelsArrayChanged().AddRaw(this, &FHierarchicalLODOutlinerModule::OnHLODLevelsArrayChangedEvent);
+	check(GUnrealEd);
+	ArrayChangedDelegate = GUnrealEd->OnHLODLevelsArrayChanged().AddRaw(this, &FHierarchicalLODOutlinerModule::OnHLODLevelsArrayChangedEvent);
 }
 
 void FHierarchicalLODOutlinerModule::ShutdownModule()
 {
-	GUnrealEd->OnHLODLevelsArrayChanged().RemoveAll(this);
+	if (GUnrealEd)
+	{
+		GUnrealEd->OnHLODLevelsArrayChanged().Remove(ArrayChangedDelegate);
+	}	
 }
 
 void FHierarchicalLODOutlinerModule::OnHLODLevelsArrayChangedEvent()
diff --git a/Engine/Source/Editor/HierarchicalLODOutliner/Public/HierarchicalLODOutlinerModule.h b/Engine/Source/Editor/HierarchicalLODOutliner/Public/HierarchicalLODOutlinerModule.h
--- a/Engine/Source/Editor/HierarchicalLODOutliner/Public/HierarchicalLODOutlinerModule.h
+++ b/Engine/Source/Editor/HierarchicalLODOutliner/Public/HierarchicalLODOutlinerModule.h
@@ -23,4 +23,7 @@ public:
 
 	/** Creates the HLOD Outliner widget */
 	virtual TSharedRef<SWidget> CreateHLODOutlinerWidget();
+
+private:
+	FDelegateHandle ArrayChangedDelegate;
 };
\ No newline at end of file
