Fix for UE-24200, HLODOutliner not un-registering with Editor for Undo callbacks
#tests undo/redo with HLOD outliner open and closed

diff --git a/Engine/Source/Editor/HierarchicalLODOutliner/Private/HLODOutliner.cpp b/Engine/Source/Editor/HierarchicalLODOutliner/Private/HLODOutliner.cpp
--- a/Engine/Source/Editor/HierarchicalLODOutliner/Private/HLODOutliner.cpp
+++ b/Engine/Source/Editor/HierarchicalLODOutliner/Private/HLODOutliner.cpp
@@ -114,8 +114,6 @@ namespace HLODOutliner
 
 		RegisterDelegates();
 
-		// Register to update when an undo/redo operation has been called to update our list of actors
-		GEditor->RegisterForUndo(this);	
 	}
 
 	TSharedRef<SWidget> SHLODOutliner::CreateButtonWidgets()
@@ -432,18 +430,20 @@ namespace HLODOutliner
 
 		// Register to be notified when properties are edited
 		FCoreDelegates::OnActorLabelChanged.AddRaw(this, &SHLODOutliner::OnActorLabelChanged);
-		
+
 		// Selection change
 		USelection::SelectionChangedEvent.AddRaw(this, &SHLODOutliner::OnLevelSelectionChanged);
 		USelection::SelectObjectEvent.AddRaw(this, &SHLODOutliner::OnLevelSelectionChanged);
-				
+
 		// HLOD related events
 		GEngine->OnHLODActorMoved().AddSP(this, &SHLODOutliner::OnHLODActorMovedEvent);
 		GEngine->OnHLODActorAdded().AddSP(this, &SHLODOutliner::OnHLODActorAddedEvent);
 		GEngine->OnHLODActorMarkedDirty().AddSP(this, &SHLODOutliner::OnHLODActorMarkedDirtyEvent);
 		GEngine->OnHLODTransitionScreenSizeChanged().AddSP(this, &SHLODOutliner::OnHLODTransitionScreenSizeChangedEvent);
 		GEngine->OnHLODLevelsArrayChanged().AddSP(this, &SHLODOutliner::OnHLODLevelsArrayChangedEvent);
-		
+
+		// Register to update when an undo/redo operation has been called to update our list of actors
+		GEditor->RegisterForUndo(this);
 	}
 
 	void SHLODOutliner::DeregisterDelegates()
@@ -466,6 +466,9 @@ namespace HLODOutliner
 		GEngine->OnHLODActorAdded().RemoveAll(this);
 		GEngine->OnHLODActorMarkedDirty().RemoveAll(this);
 		GEngine->OnHLODLevelsArrayChanged().RemoveAll(this);
+
+		// Deregister for Undo callbacks
+		GEditor->UnregisterForUndo(this);
 	}
 
 	void SHLODOutliner::ForceViewLODActor()
