Copying //UE4/Dev-Framework to Dev-Main (//UE4/Dev-Main) @ 2855699

diff --git a/Engine/Source/Developer/SimplygonSwarm/Private/SimplygonSwarm.cpp b/Engine/Source/Developer/SimplygonSwarm/Private/SimplygonSwarm.cpp
--- a/Engine/Source/Developer/SimplygonSwarm/Private/SimplygonSwarm.cpp
+++ b/Engine/Source/Developer/SimplygonSwarm/Private/SimplygonSwarm.cpp
@@ -43,7 +43,7 @@
  #endif
 
 
- static const TCHAR* SPL_TEMPLATE_REMESHING = TEXT("{\"Header\":{\"SPLVersion\":\"7.0\",\"ClientName\":\"UE4\",\"ClientVersion\":\"UE4.9\",\"SimplygonVersion\":\"7.0\"},\"ProcessGraph\":{\"Type\":\"ContainerNode\",\"Name\":\"Node\",\"Children\":[{\"Processor\":{\"RemeshingSettings\":{\"CuttingPlaneSelectionSetName\":0,\"EmptySpaceOverride\":0.0,\"MaxTriangleSize\":32,\"OnScreenSize\":%d,\"ProcessSelectionSetName\":\"\",\"SurfaceTransferMode\":1,\"TransferColors\":false,\"TransferNormals\":false,\"UseCuttingPlanes\":%s,\"UseEmptySpaceOverride\":false,\"Enabled\":true%s},\"MappingImageSettings\":{\"AutomaticTextureSizeMultiplier\":1.0,\"ChartAggregatorMode\":0,\"ChartAggregatorOriginalTexCoordLevel\":0,\"ChartAggregatorUseAreaWeighting\":true,\"ChartAggregatorKeepOriginalChartProportions\":true,\"ChartAggregatorKeepOriginalChartProportionsFromChannel\":\"\",\"ChartAggregatorKeepOriginalChartSizes\":false,\"ChartAggregatorSeparateOverlappingCharts\":true,\"ForcePower2Texture\":true,\"GenerateMappingImage\":true,\"GenerateTangents\":true,\"GenerateTexCoords\":true,\"GutterSpace\":4,\"Height\":%d,\"MaximumLayers\":3,\"MultisamplingLevel\":3,\"ParameterizerMaxStretch\":6.0,\"ParameterizerUseVertexWeights\":false,\"ParameterizerUseVisibilityWeights\":false,\"TexCoordGeneratorType\":0,\"TexCoordLevel\":255,\"UseAutomaticTextureSize\":false,\"UseFullRetexturing\":false,\"Width\":%d,\"Enabled\":true},%s\"Type\":\"RemeshingProcessor\"},\"MaterialCaster\":[%s],\"DefaultTBNType\":2,\"AllowGPUAcceleration\":false,\"Type\":\"ProcessNode\",\"Name\":\"Node\",\"Children\":[{\"Format\":\"ssf\",\"Type\":\"WriteNode\",\"Name\":\"outputlod_0\",\"Children\":[]}]}]}}");
+ static const TCHAR* SPL_TEMPLATE_REMESHING = TEXT("{\"Header\":{\"SPLVersion\":\"7.0\",\"ClientName\":\"UE4\",\"ClientVersion\":\"UE4.9\",\"SimplygonVersion\":\"7.0\"},\"ProcessGraph\":{\"Type\":\"ContainerNode\",\"Name\":\"Node\",\"Children\":[{\"Processor\":{\"RemeshingSettings\":{\"CuttingPlaneSelectionSetName\":0,\"EmptySpaceOverride\":0.0,\"MaxTriangleSize\":32,\"OnScreenSize\":%d,\"ProcessSelectionSetName\":\"\",\"SurfaceTransferMode\":1,\"TransferColors\":false,\"TransferNormals\":false,\"UseCuttingPlanes\":false,\"UseEmptySpaceOverride\":false,\"Enabled\":true%s},\"MappingImageSettings\":{\"AutomaticTextureSizeMultiplier\":1.0,\"ChartAggregatorMode\":0,\"ChartAggregatorOriginalTexCoordLevel\":0,\"ChartAggregatorUseAreaWeighting\":true,\"ChartAggregatorKeepOriginalChartProportions\":true,\"ChartAggregatorKeepOriginalChartProportionsFromChannel\":\"\",\"ChartAggregatorKeepOriginalChartSizes\":false,\"ChartAggregatorSeparateOverlappingCharts\":true,\"ForcePower2Texture\":true,\"GenerateMappingImage\":true,\"GenerateTangents\":true,\"GenerateTexCoords\":true,\"GutterSpace\":4,\"Height\":%d,\"MaximumLayers\":3,\"MultisamplingLevel\":3,\"ParameterizerMaxStretch\":6.0,\"ParameterizerUseVertexWeights\":false,\"ParameterizerUseVisibilityWeights\":false,\"TexCoordGeneratorType\":0,\"TexCoordLevel\":255,\"UseAutomaticTextureSize\":false,\"UseFullRetexturing\":false,\"Width\":%d,\"Enabled\":true},\"Type\":\"RemeshingProcessor\"},\"MaterialCaster\":[%s],\"DefaultTBNType\":2,\"AllowGPUAcceleration\":false,\"Type\":\"ProcessNode\",\"Name\":\"Node\",\"Children\":[{\"Format\":\"ssf\",\"Type\":\"WriteNode\",\"Name\":\"outputlod_0\",\"Children\":[]}]}]}}");
 
 
  static const TCHAR* SPL_TEMPLATE_COLORCASTER = TEXT("{\"BakeOpacityInAlpha\":false,\"ColorType\":\"%s\",\"Dilation\":8,\"FillMode\":2,\"IsSRGB\":true,\"OutputChannelBitDepth\":8,\"OutputChannels\":4,\"Type\":\"ColorCaster\",\"Name\":\"%s\",\"Channel\":\"%s\",\"DitherType\":0}");
@@ -431,53 +431,11 @@
 
 		 if (InProxySettings.bRecalculateNormals)
 			 AdditionalSettings = FString::Printf(TEXT("%s,%s"), *AdditionalSettings, *RecalNormals);
-
-		 //Note : The Simplygon API now supports multiple clipping planes. SPL has recently added support for this.
-		 // 
-		 FString CuttingPlaneSetting;
-
-		 if (InProxySettings.bUseClippingPlane)
-		 {
-			 // Note : An arbitary plane can be define using a point (position) and a normal (direction)
-			 // Since UE  currently only has axis aligned plane. We need to convert values for SPL
-
-			 FVector OutPoint, OutNormal;
-
-			 GetAxisAlignedVectorsForCuttingPlanes(InProxySettings, OutPoint, OutNormal);
-
-			 FString CuttingPlane = FString::Printf(CUTTING_PLANE, OutPoint.X, OutPoint.Y, OutPoint.Z, OutNormal.X, OutNormal.Y, OutNormal.Z);
-			 CuttingPlaneSetting = FString::Printf(CUTTING_PLANE_SETTINGS, *CuttingPlane);
-		 }
-		 
-		 OutSplText = FString::Printf(SPL_TEMPLATE_REMESHING, InProxySettings.ScreenSize, InProxySettings.bUseClippingPlane ? TEXT("true") :TEXT("false"), *AdditionalSettings, ImageSizes.X, ImageSizes.Y, *CuttingPlaneSetting, *CasterSPLTempalte);
-
+		 		 
+		 OutSplText = FString::Printf(SPL_TEMPLATE_REMESHING, InProxySettings.ScreenSize, *AdditionalSettings, ImageSizes.X, ImageSizes.Y, *CasterSPLTempalte);
 	 }
 	 
-
-	 void GetAxisAlignedVectorsForCuttingPlanes(const struct FMeshProxySettings& InProxySettings, 
-		 FVector& OutPoint, 
-		 FVector& OutNormal)
-	 {
-		 
-		 // 0 -> X , 1 -> Y, 2 -> Z
-		 if (InProxySettings.AxisIndex == 0)
-		 {
-			 OutNormal.X = InProxySettings.bPlaneNegativeHalfspace ? -1 : 1 ;
-			 OutPoint.X = InProxySettings.ClippingLevel;
-			 
-		 }
-		 else if (InProxySettings.AxisIndex == 1)
-		 {
-			 OutNormal.Y = InProxySettings.bPlaneNegativeHalfspace ? -1 : 1;
-			 OutPoint.Y = InProxySettings.ClippingLevel;
-		 }
-		 else
-		 {
-			 OutNormal.Z = InProxySettings.bPlaneNegativeHalfspace ? -1 : 1;
-			 OutPoint.Z = InProxySettings.ClippingLevel;
-			 //default to z up 
-		 }
-	 }
+	
 	 /*
 	 Write the SPL string to file
 	 */
diff --git a/Engine/Source/Editor/ComponentVisualizers/Private/ConstraintComponentVisualizer.cpp b/Engine/Source/Editor/ComponentVisualizers/Private/ConstraintComponentVisualizer.cpp
--- a/Engine/Source/Editor/ComponentVisualizers/Private/ConstraintComponentVisualizer.cpp
+++ b/Engine/Source/Editor/ComponentVisualizers/Private/ConstraintComponentVisualizer.cpp
@@ -20,9 +20,7 @@ void FConstraintComponentVisualizer::DrawVisualization( const UActorComponent* C
 		if(Instance.IsValidConstraintInstance())
 		{
 			FTransform BodyTransform1 = ConstraintComp->GetBodyTransform(EConstraintFrame::Frame1);
-			BodyTransform1.RemoveScaling();
 			FTransform BodyTransform2 = ConstraintComp->GetBodyTransform(EConstraintFrame::Frame2);
-			BodyTransform2.RemoveScaling();
 
 			Con1Frame = Instance.GetRefFrame(EConstraintFrame::Frame1) * BodyTransform1;
 			Con2Frame = Instance.GetRefFrame(EConstraintFrame::Frame2) * BodyTransform2;
diff --git a/Engine/Source/Editor/DetailCustomizations/Private/AnimMontageSegmentDetails.cpp b/Engine/Source/Editor/DetailCustomizations/Private/AnimMontageSegmentDetails.cpp
--- a/Engine/Source/Editor/DetailCustomizations/Private/AnimMontageSegmentDetails.cpp
+++ b/Engine/Source/Editor/DetailCustomizations/Private/AnimMontageSegmentDetails.cpp
@@ -101,15 +101,14 @@ SAnimationSegmentViewport::~SAnimationSegmentViewport()
 	// clean up components
 	if (PreviewComponent)
 	{
-		for (int32 I=PreviewComponent->AttachChildren.Num()-1; I >= 0; --I) // Iterate backwards because Cleancomponent will remove from AttachChildren
+		for (int32 I=PreviewComponent->GetAttachChildren().Num()-1; I >= 0; --I) // Iterate backwards because CleanupComponent will remove from AttachChildren
 		{
 			// PreviewComponet will be cleaned up by PreviewScene, 
 			// but if anything is attached, it won't be cleaned up, 
 			// so we'll need to clean them up manually
-			CleanupComponent(PreviewComponent->AttachChildren[I]);
+			CleanupComponent(PreviewComponent->GetAttachChildren()[I]);
 		}
-
-		PreviewComponent->AttachChildren.Empty();
+		check(PreviewComponent->GetAttachChildren().Num() == 0);
 	}
 
 	// Close viewport
@@ -123,12 +122,12 @@ void SAnimationSegmentViewport::CleanupComponent(USceneComponent* Component)
 {
 	if (Component)
 	{
-		for (int32 I=0; I<Component->AttachChildren.Num(); ++I)
+		for (int32 I = Component->GetAttachChildren().Num() - 1; I >= 0; --I) // Iterate backwards because CleanupComponent will remove from AttachChildren
 		{
-			CleanupComponent(Component->AttachChildren[I]);
+			CleanupComponent(Component->GetAttachChildren()[I]);
 		}
+		check(Component->GetAttachChildren().Num() == 0);
 
-		Component->AttachChildren.Empty();
 		Component->DestroyComponent();
 	}
 }
diff --git a/Engine/Source/Editor/DetailCustomizations/Private/AnimSequenceDetails.cpp b/Engine/Source/Editor/DetailCustomizations/Private/AnimSequenceDetails.cpp
--- a/Engine/Source/Editor/DetailCustomizations/Private/AnimSequenceDetails.cpp
+++ b/Engine/Source/Editor/DetailCustomizations/Private/AnimSequenceDetails.cpp
@@ -366,15 +366,14 @@ SAnimationRefPoseViewport::~SAnimationRefPoseViewport()
 	// clean up components
 	if (PreviewComponent)
 	{
-		for (int32 I=PreviewComponent->AttachChildren.Num()-1; I >= 0; --I) // Iterate backwards because Cleancomponent will remove from AttachChildren
+		for (int32 I=PreviewComponent->GetAttachChildren().Num()-1; I >= 0; --I) // Iterate backwards because CleanupComponent will remove from AttachChildren
 		{
 			// PreviewComponet will be cleaned up by PreviewScene, 
 			// but if anything is attached, it won't be cleaned up, 
 			// so we'll need to clean them up manually
-			CleanupComponent(PreviewComponent->AttachChildren[I]);
+			CleanupComponent(PreviewComponent->GetAttachChildren()[I]);
 		}
-
-		PreviewComponent->AttachChildren.Empty();
+		check(PreviewComponent->GetAttachChildren().Num() == 0);
 	}
 
 	// Close viewport
@@ -388,12 +387,12 @@ void SAnimationRefPoseViewport::CleanupComponent(USceneComponent* Component)
 {
 	if (Component)
 	{
-		for (int32 I=0; I<Component->AttachChildren.Num(); ++I)
+		for (int32 I = Component->GetAttachChildren().Num() - 1; I >= 0; --I) // Iterate backwards because CleanupComponent will remove from AttachChildren
 		{
-			CleanupComponent(Component->AttachChildren[I]);
+			CleanupComponent(Component->GetAttachChildren()[I]);
 		}
+		check(Component->GetAttachChildren().Num() == 0);
 
-		Component->AttachChildren.Empty();
 		Component->DestroyComponent();
 	}
 }
diff --git a/Engine/Source/Editor/DetailCustomizations/Private/BodyInstanceCustomization.cpp b/Engine/Source/Editor/DetailCustomizations/Private/BodyInstanceCustomization.cpp
--- a/Engine/Source/Editor/DetailCustomizations/Private/BodyInstanceCustomization.cpp
+++ b/Engine/Source/Editor/DetailCustomizations/Private/BodyInstanceCustomization.cpp
@@ -75,43 +75,63 @@ void FBodyInstanceCustomization::AddCollisionCategory(TSharedRef<class IProperty
 	]
 	.ValueContent()
 	[
-		SNew(SHorizontalBox)
-		+SHorizontalBox::Slot()
-		.VAlign(VAlign_Center)
+		SNew(SVerticalBox)
+		+ SVerticalBox::Slot()
+		.Padding(0.f, 0.f, 10.f, 0.f)
 		[
-			SAssignNew(CollsionProfileComboBox, SComboBox< TSharedPtr<FString> >)
-			.OptionsSource(&CollisionProfileComboList)
-			.OnGenerateWidget(this, &FBodyInstanceCustomization::MakeCollisionProfileComboWidget)
-			.OnSelectionChanged(this, &FBodyInstanceCustomization::OnCollisionProfileChanged, &CollisionGroup )
-			.OnComboBoxOpening(this, &FBodyInstanceCustomization::OnCollisionProfileComboOpening)
-			.InitiallySelectedItem(DisplayName)
-			.IsEnabled( FSlateApplication::Get().GetNormalExecutionAttribute() )
-			.ContentPadding(2)
-			.Content()
+			SNew(SHorizontalBox)
+			.Visibility(this, &FBodyInstanceCustomization::IsCollisionPresetVisible)
+			.IsEnabled(this, &FBodyInstanceCustomization::IsCollisionEnabled)
+			+ SHorizontalBox::Slot()
+			.VAlign(VAlign_Center)
 			[
-				SNew( STextBlock )
-				.Text(this, &FBodyInstanceCustomization::GetCollisionProfileComboBoxContent)
-				.Font( IDetailLayoutBuilder::GetDetailFont() )
-				.ToolTipText(this, &FBodyInstanceCustomization::GetCollisionProfileComboBoxToolTip)
+				SAssignNew(CollsionProfileComboBox, SComboBox< TSharedPtr<FString> >)
+				.OptionsSource(&CollisionProfileComboList)
+				.OnGenerateWidget(this, &FBodyInstanceCustomization::MakeCollisionProfileComboWidget)
+				.OnSelectionChanged(this, &FBodyInstanceCustomization::OnCollisionProfileChanged, &CollisionGroup)
+				.OnComboBoxOpening(this, &FBodyInstanceCustomization::OnCollisionProfileComboOpening)
+				.InitiallySelectedItem(DisplayName)
+				.ContentPadding(2)
+				.Content()
+				[
+					SNew(STextBlock)
+					.Text(this, &FBodyInstanceCustomization::GetCollisionProfileComboBoxContent)
+					.Font(IDetailLayoutBuilder::GetDetailFont())
+					.ToolTipText(this, &FBodyInstanceCustomization::GetCollisionProfileComboBoxToolTip)
+				]
 			]
+
+			+ SHorizontalBox::Slot()
+				.VAlign(VAlign_Center)
+				.Padding(2.0f)
+				.AutoWidth()
+				[
+					SNew(SButton)
+					.OnClicked(this, &FBodyInstanceCustomization::SetToDefaultProfile)
+					.ContentPadding(0.f)
+					.ToolTipText(LOCTEXT("ResetToDefaultToolTip", "Reset to Default"))
+					.ButtonStyle(FEditorStyle::Get(), "NoBorder")
+					.IsEnabled(this, &FBodyInstanceCustomization::IsCollisionEnabled)
+					.Content()
+					[
+						SNew(SImage)
+						.Image(FEditorStyle::GetBrush("PropertyWindow.DiffersFromDefault"))
+					]
+				]
 		]
-		+SHorizontalBox::Slot()
-		.VAlign(VAlign_Center)
-		.Padding( 2.0f )
-		.AutoWidth()
-		[
-			SNew(SButton)
-			.OnClicked(this, &FBodyInstanceCustomization::SetToDefaultProfile)
-			.Visibility(this, &FBodyInstanceCustomization::ShouldShowResetToDefaultProfile)
-			.ContentPadding(0.f)
-			.ToolTipText(LOCTEXT("ResetToDefaultToolTip", "Reset to Default"))
-			.ButtonStyle( FEditorStyle::Get(), "NoBorder" )
-			.Content()
+
+		+ SVerticalBox::Slot()
+			.Padding(0.f, 0.f, 10.f, 0.f)
 			[
-				SNew(SImage)
-				.Image( FEditorStyle::GetBrush("PropertyWindow.DiffersFromDefault") )
+				SNew(STextBlock)
+				.Text(LOCTEXT("DefaultCollision", "Default Collision"))
+				.Visibility(this, &FBodyInstanceCustomization::IsDefaultCollisionVisible)
+				.Font(IDetailLayoutBuilder::GetDetailFont())
+				.ToolTipText(LOCTEXT("DefaultCollisionToolTip", "Default Collision. See the StaticMesh asset to find out collision settings"))
 			]
-		]
+
+
+		
 	];
 
 	CollisionGroup.ToggleExpansion(bDisplayAdvancedCollisionSettings);
@@ -122,6 +142,8 @@ void FBodyInstanceCustomization::AddCollisionCategory(TSharedRef<class IProperty
 BEGIN_SLATE_FUNCTION_BUILD_OPTIMIZATION
 void FBodyInstanceCustomization::CustomizeChildren( TSharedRef<class IPropertyHandle> StructPropertyHandle, class IDetailChildrenBuilder& StructBuilder, IPropertyTypeCustomizationUtils& StructCustomizationUtils )
 {
+	BodyInstanceHandle = StructPropertyHandle;
+
 	// copy all bodyinstances I'm accessing right now
 	TArray<void*> StructPtrs;
 	StructPropertyHandle->AccessRawData(StructPtrs);
@@ -134,6 +156,31 @@ void FBodyInstanceCustomization::CustomizeChildren( TSharedRef<class IPropertyHa
 		BodyInstances[Iter.GetIndex()] = (FBodyInstance*)(*Iter);
 	}
 
+	
+	// get all parent instances
+	TSharedPtr<IPropertyHandle> ParentPropertyHandle = StructPropertyHandle->GetParentHandle();
+	bool bFoundValidProperty = false;	//Need to go up root until we find a valid property. This is a bad hack but is needed because GetChildHandle assumes property is always non-null
+	while(ParentPropertyHandle.IsValid() && !bFoundValidProperty)
+	{
+		if(ParentPropertyHandle->GetProperty())
+		{
+			bFoundValidProperty = true;
+		}
+		else
+		{
+			ParentPropertyHandle = ParentPropertyHandle->GetParentHandle();
+		}
+	}
+
+
+	if(ParentPropertyHandle.IsValid())
+	{
+		
+		UseDefaultCollisionHandle = ParentPropertyHandle->GetChildHandle(GET_MEMBER_NAME_CHECKED(UStaticMeshComponent, bUseDefaultCollision));
+		StaticMeshHandle = ParentPropertyHandle->GetChildHandle(GET_MEMBER_NAME_CHECKED(UStaticMeshComponent, StaticMesh));
+	}
+	
+
 	AddCollisionCategory(StructPropertyHandle, StructBuilder, StructCustomizationUtils);
 }
 END_SLATE_FUNCTION_BUILD_OPTIMIZATION
@@ -259,6 +306,7 @@ void FBodyInstanceCustomization::CreateCustomCollisionSetup( TSharedRef<class IP
 	}
 
 	int32 TotalNumChildren = ValidCollisionChannels.Num();
+	TAttribute<bool> CollisionEnabled(this, &FBodyInstanceCustomization::IsCollisionEnabled );
 	TAttribute<bool> CustomCollisionEnabled( this, &FBodyInstanceCustomization::ShouldEnableCustomCollisionSetup );
 	TAttribute<EVisibility> CustomCollisionVisibility(this, &FBodyInstanceCustomization::ShouldShowCustomCollisionSetup);
 
@@ -273,7 +321,6 @@ void FBodyInstanceCustomization::CreateCustomCollisionSetup( TSharedRef<class IP
 	if (!StructPropertyHandle->GetProperty()->GetBoolMetaData(TEXT("HideObjectType")))
 	{
 		CollisionGroup.AddWidgetRow()
-		.IsEnabled(CustomCollisionEnabled)
 		.Visibility(CustomCollisionVisibility)
 		.NameContent()
 		[
@@ -286,7 +333,7 @@ void FBodyInstanceCustomization::CreateCustomCollisionSetup( TSharedRef<class IP
 			.OnGenerateWidget(this, &FBodyInstanceCustomization::MakeObjectTypeComboWidget)
 			.OnSelectionChanged(this, &FBodyInstanceCustomization::OnObjectTypeChanged)
 			.InitiallySelectedItem(ObjectTypeComboList[IndexSelected])
-			.IsEnabled(FSlateApplication::Get().GetNormalExecutionAttribute())
+			.IsEnabled(CustomCollisionEnabled)
 			.ContentPadding(2)
 			.Content()
 			[
@@ -804,12 +851,56 @@ EVisibility FBodyInstanceCustomization::ShouldShowResetToDefaultResponse(int32 V
 	return EVisibility::Hidden;
 }
 
+bool FBodyInstanceCustomization::AreAllCollisionUsingDefault() const
+{
+	if(UseDefaultCollisionHandle.IsValid())
+	{
+		bool bUseDefaultCollision = false;
+		if (UseDefaultCollisionHandle->GetValue(bUseDefaultCollision) == FPropertyAccess::Result::Success && bUseDefaultCollision)
+		{
+			return true;
+		}
+	}
+	
+	return false;
+}
+
+EVisibility FBodyInstanceCustomization::IsCollisionPresetVisible() const
+{
+	if (StaticMeshHandle.IsValid())
+	{
+		UObject* StaticMesh = nullptr;
+		if (StaticMeshHandle->GetValue(StaticMesh) == FPropertyAccess::Result::Success && StaticMesh)
+		{
+			return AreAllCollisionUsingDefault() ? EVisibility::Collapsed : EVisibility::Visible;
+		}
+	}
+
+	return EVisibility::Visible;
+}
+
+EVisibility FBodyInstanceCustomization::IsDefaultCollisionVisible() const
+{
+	return IsCollisionPresetVisible()  == EVisibility::Visible ? EVisibility::Collapsed : EVisibility::Visible;
+}
+
+bool FBodyInstanceCustomization::IsCollisionEnabled() const
+{
+	bool bEnabled = false;
+	if(BodyInstanceHandle.IsValid())
+	{
+		bEnabled = !BodyInstanceHandle->IsEditConst() && FSlateApplication::Get().GetNormalExecutionAttribute().Get();
+	}
+
+	return bEnabled;
+}
+
 bool FBodyInstanceCustomization::ShouldEnableCustomCollisionSetup() const
 {
 	FName ProfileName;
 	if (CollisionProfileNameHandle->GetValue(ProfileName) == FPropertyAccess::Result::Success && FBodyInstance::IsValidCollisionProfileName(ProfileName) == false)
 	{
-		return true;
+		return IsCollisionEnabled();
 	}
 
 	return false;
@@ -817,7 +908,7 @@ bool FBodyInstanceCustomization::ShouldEnableCustomCollisionSetup() const
 
 EVisibility FBodyInstanceCustomization::ShouldShowCustomCollisionSetup() const
 {
-	return EVisibility::Visible;//return (bDisplayAdvancedCollisionSettings)? EVisibility::Visible : EVisibility::Hidden;
+	return AreAllCollisionUsingDefault() ? EVisibility::Hidden : EVisibility::Visible;
 }
 
 FText FBodyInstanceCustomization::GetCollisionProfileComboBoxContent() const
@@ -1226,6 +1317,26 @@ EVisibility FBodyInstanceCustomizationHelper::IsAutoWeldVisible() const
 	return EVisibility::Visible;
 }
 
+void FBodyInstanceCustomizationHelper::OnSetBodyMass(float BodyMass, ETextCommit::Type Commit)
+{
+	UPrimitiveComponent* Comp = nullptr;
+	UBodySetup* BS = nullptr;
+
+	for (auto ObjectIt = ObjectsCustomized.CreateConstIterator(); ObjectIt; ++ObjectIt)
+	{
+		if (ObjectIt->IsValid() && (*ObjectIt)->IsA(UPrimitiveComponent::StaticClass()))
+		{
+			Comp = Cast<UPrimitiveComponent>(ObjectIt->Get());
+			Comp->SetMassOverrideInKg(NAME_None, BodyMass);
+		}
+		else if (ObjectIt->IsValid() && (*ObjectIt)->IsA(UBodySetup::StaticClass()))
+		{
+			BS = Cast<UBodySetup>(ObjectIt->Get());
+			BS->DefaultInstance.SetMassOverride(BodyMass);
+		}
+	}
+}
+
 
 TOptional<float> FBodyInstanceCustomizationHelper::OnGetBodyMass() const
 {
@@ -1237,36 +1348,27 @@ TOptional<float> FBodyInstanceCustomizationHelper::OnGetBodyMass() const
 
 	for (auto ObjectIt = ObjectsCustomized.CreateConstIterator(); ObjectIt; ++ObjectIt)
 	{
+		float NewMass = 0.f;
 		if (ObjectIt->IsValid() && (*ObjectIt)->IsA(UPrimitiveComponent::StaticClass()))
 		{
 			Comp = Cast<UPrimitiveComponent>(ObjectIt->Get());
-
-			float CompMass = Comp->CalculateMass();
-			Comp->BodyInstance.MassInKg = CompMass;	//update the component's override mass to be the calculated one
-			if (Mass == 0.0f || FMath::Abs(Mass - CompMass) < 0.1f)
-			{
-				Mass = CompMass;
-			}
-			else
-			{
-				bMultipleValue = true;
-				break;
-			}
+			NewMass = Comp->CalculateMass();
 		}
 		else if (ObjectIt->IsValid() && (*ObjectIt)->IsA(UBodySetup::StaticClass()))
 		{
 			BS = Cast<UBodySetup>(ObjectIt->Get());
 
-			float BSMass = BS->CalculateMass();
-			if (Mass == 0.0f || FMath::Abs(Mass - BSMass) < 0.1f)
-			{
-				Mass = BSMass;
-			}
-			else
-			{
-				bMultipleValue = true;
-				break;
-			}
+			NewMass = BS->CalculateMass();
+		}
+
+		if (Mass == 0.0f || FMath::Abs(Mass - NewMass) < SMALL_NUMBER)
+		{
+			Mass = NewMass;
+		}
+		else
+		{
+			bMultipleValue = true;
+			break;
 		}
 	}
 
@@ -1299,12 +1401,12 @@ void FBodyInstanceCustomizationHelper::AddMassInKg(IDetailCategoryBuilder& Physi
 {
 	if (bDisplayMass)
 	{
-		TSharedRef<IPropertyHandle> MassInKGHandle = BodyInstanceHandler->GetChildHandle(GET_MEMBER_NAME_CHECKED(FBodyInstance, MassInKg)).ToSharedRef();
-
-		PhysicsCategory.AddProperty(MassInKGHandle).CustomWidget()
+		TSharedRef<IPropertyHandle> MassInKgOverrideHandle = BodyInstanceHandler->GetChildHandle(GET_MEMBER_NAME_CHECKED(FBodyInstance, MassInKgOverride)).ToSharedRef();
+		
+		PhysicsCategory.AddProperty(MassInKgOverrideHandle).CustomWidget()
 		.NameContent()
 		[
-			MassInKGHandle->CreatePropertyNameWidget()
+			MassInKgOverrideHandle->CreatePropertyNameWidget()
 		]
 		.ValueContent()
 		[
@@ -1313,22 +1415,10 @@ void FBodyInstanceCustomizationHelper::AddMassInKg(IDetailCategoryBuilder& Physi
 			.Padding(0.f, 0.f, 10.f, 0.f)
 			[
 				SNew(SNumericEntryBox<float>)
-				.IsEnabled(false)
+				.IsEnabled(&FBodyInstanceCustomizationHelper::IsBodyMassEnabled)
 				.Font(IDetailLayoutBuilder::GetDetailFont())
 				.Value(this, &FBodyInstanceCustomizationHelper::OnGetBodyMass)
-				.Visibility(this, &FBodyInstanceCustomizationHelper::IsMassVisible, false)
-			]
-
-			+ SVerticalBox::Slot()
-			.Padding(0.f, 0.f, 10.f, 0.f)
-			[
-				SNew(SVerticalBox)
-				.Visibility(this, &FBodyInstanceCustomizationHelper::IsMassVisible, true)
-				+ SVerticalBox::Slot()
-				.AutoHeight()
-				[
-					MassInKGHandle->CreatePropertyValueWidget()
-				]
+				.OnValueCommitted(this, &FBodyInstanceCustomizationHelper::OnSetBodyMass)
 			]
 		];
 	}
diff --git a/Engine/Source/Editor/DetailCustomizations/Private/BodyInstanceCustomization.h b/Engine/Source/Editor/DetailCustomizations/Private/BodyInstanceCustomization.h
--- a/Engine/Source/Editor/DetailCustomizations/Private/BodyInstanceCustomization.h
+++ b/Engine/Source/Editor/DetailCustomizations/Private/BodyInstanceCustomization.h
@@ -58,6 +58,13 @@ private:
 	bool ShouldEnableCustomCollisionSetup() const;
 	EVisibility ShouldShowCustomCollisionSetup() const;
 
+	bool IsCollisionEnabled() const;
+
+	// whether we can edit collision or if we're getting it from a default
+	EVisibility IsCollisionPresetVisible() const;
+	EVisibility IsDefaultCollisionVisible() const;
+	bool AreAllCollisionUsingDefault() const;
+
 	// utility functions between property and struct
 	void CreateCustomCollisionSetup( TSharedRef<class IPropertyHandle> StructPropertyHandle, class IDetailGroup& CollisionGroup );
 	void SetCollisionResponseContainer(const FCollisionResponseContainer& ResponseContainer);
@@ -72,10 +79,13 @@ private:
 
 private:
 	// property handles
+	TSharedPtr<IPropertyHandle> BodyInstanceHandle;
 	TSharedPtr<IPropertyHandle> CollisionProfileNameHandle;
 	TSharedPtr<IPropertyHandle> CollisionEnabledHandle;
 	TSharedPtr<IPropertyHandle> ObjectTypeHandle;
 	TSharedPtr<IPropertyHandle> CollisionResponsesHandle;
+	TSharedPtr<IPropertyHandle> UseDefaultCollisionHandle;
+	TSharedPtr<IPropertyHandle> StaticMeshHandle;
 
 	// widget related variables
 	TSharedPtr<class SComboBox< TSharedPtr<FString> > > CollsionProfileComboBox;
@@ -109,6 +119,7 @@ private:
 	bool IsUseAsyncEditable() const;
 
 	TOptional<float> OnGetBodyMass() const;
+	void OnSetBodyMass(float InBodyMass, ETextCommit::Type Commit);
 	bool IsBodyMassReadOnly() const;
 	EVisibility IsMassVisible(bool bOverrideMass) const;
 	bool IsBodyMassEnabled() const { return !IsBodyMassReadOnly(); }
diff --git a/Engine/Source/Editor/DetailCustomizations/Private/HierarchicalSimplificationCustomizations.cpp b/Engine/Source/Editor/DetailCustomizations/Private/HierarchicalSimplificationCustomizations.cpp
--- a/Engine/Source/Editor/DetailCustomizations/Private/HierarchicalSimplificationCustomizations.cpp
+++ b/Engine/Source/Editor/DetailCustomizations/Private/HierarchicalSimplificationCustomizations.cpp
@@ -51,8 +51,9 @@ void FHierarchicalSimplificationCustomizations::CustomizeChildren( TSharedRef<IP
 	SimplifyMeshPropertyHandle = PropertyHandles.FindChecked(GET_MEMBER_NAME_CHECKED(FHierarchicalSimplification, bSimplifyMesh));
 	TSharedPtr< IPropertyHandle > ProxyMeshSettingPropertyHandle = PropertyHandles.FindChecked(GET_MEMBER_NAME_CHECKED(FHierarchicalSimplification, ProxySetting));
 	TSharedPtr< IPropertyHandle > MergeMeshSettingPropertyHandle = PropertyHandles.FindChecked(GET_MEMBER_NAME_CHECKED(FHierarchicalSimplification, MergeSetting));
+	TSharedPtr< IPropertyHandle > TransitionScreenSizePropertyHandle = PropertyHandles.FindChecked(GET_MEMBER_NAME_CHECKED(FHierarchicalSimplification, TransitionScreenSize));
 
-	for( auto Iter(PropertyHandles.CreateConstIterator()); Iter; ++Iter  )
+	for (auto Iter(PropertyHandles.CreateConstIterator()); Iter; ++Iter)
 	{
 		// Handle special property cases (done inside the loop to maintain order according to the struct
 		if (Iter.Value() == SimplifyMeshPropertyHandle)
@@ -70,11 +71,15 @@ void FHierarchicalSimplificationCustomizations::CustomizeChildren( TSharedRef<IP
 			IDetailPropertyRow& SettingsRow = MergeGroup.AddPropertyRow(Iter.Value().ToSharedRef());
 			SettingsRow.Visibility(TAttribute<EVisibility>(this, &FHierarchicalSimplificationCustomizations::IsMergeMeshSettingVisible));
 		}
+		else  if (Iter.Value() == TransitionScreenSizePropertyHandle)
+		{
+			IDetailPropertyRow& SettingsRow = MergeGroup.AddPropertyRow(Iter.Value().ToSharedRef());
+		}
 		else
 		{
 			IDetailPropertyRow& SettingsRow = ClusterGroup.AddPropertyRow(Iter.Value().ToSharedRef());
 		}
-	}	
+	}
 }
 
 EVisibility FHierarchicalSimplificationCustomizations::IsSimplifyMeshVisible() const
diff --git a/Engine/Source/Editor/DetailCustomizations/Private/PhysicsConstraintComponentDetails.cpp b/Engine/Source/Editor/DetailCustomizations/Private/PhysicsConstraintComponentDetails.cpp
--- a/Engine/Source/Editor/DetailCustomizations/Private/PhysicsConstraintComponentDetails.cpp
+++ b/Engine/Source/Editor/DetailCustomizations/Private/PhysicsConstraintComponentDetails.cpp
@@ -161,6 +161,7 @@ void FPhysicsConstraintComponentDetails::AddLinearLimits(IDetailLayoutBuilder& D
 
 
 	LinearLimitCat.AddProperty(ConstraintInstance->GetChildHandle(GET_MEMBER_NAME_CHECKED(FConstraintInstance, LinearLimitSize)).ToSharedRef()).IsEnabled(TAttribute<bool>::Create(TAttribute<bool>::FGetter::CreateSP(this, &FPhysicsConstraintComponentDetails::IsPropertyEnabled, EPropertyType::LinearLimit)));
+	LinearLimitCat.AddProperty(ConstraintInstance->GetChildHandle(GET_MEMBER_NAME_CHECKED(FConstraintInstance, bScaleLinearLimits)).ToSharedRef()).IsEnabled(TAttribute<bool>::Create(TAttribute<bool>::FGetter::CreateSP(this, &FPhysicsConstraintComponentDetails::IsPropertyEnabled, EPropertyType::LinearLimit)));
 	LinearLimitCat.AddProperty(ConstraintInstance->GetChildHandle(GET_MEMBER_NAME_CHECKED(FConstraintInstance, bLinearLimitSoft)).ToSharedRef()).IsEnabled(TAttribute<bool>::Create(TAttribute<bool>::FGetter::CreateSP(this, &FPhysicsConstraintComponentDetails::IsPropertyEnabled, EPropertyType::LinearLimit)));
 	LinearLimitCat.AddProperty(ConstraintInstance->GetChildHandle(GET_MEMBER_NAME_CHECKED(FConstraintInstance, LinearLimitStiffness)).ToSharedRef()).IsEnabled(TAttribute<bool>::Create(TAttribute<bool>::FGetter::CreateSP(this, &FPhysicsConstraintComponentDetails::IsPropertyEnabled, EPropertyType::LinearLimit)));
 	LinearLimitCat.AddProperty(ConstraintInstance->GetChildHandle(GET_MEMBER_NAME_CHECKED(FConstraintInstance, LinearLimitDamping)).ToSharedRef()).IsEnabled(TAttribute<bool>::Create(TAttribute<bool>::FGetter::CreateSP(this, &FPhysicsConstraintComponentDetails::IsPropertyEnabled, EPropertyType::LinearLimit)));
diff --git a/Engine/Source/Editor/DetailCustomizations/Private/SceneComponentDetails.cpp b/Engine/Source/Editor/DetailCustomizations/Private/SceneComponentDetails.cpp
--- a/Engine/Source/Editor/DetailCustomizations/Private/SceneComponentDetails.cpp
+++ b/Engine/Source/Editor/DetailCustomizations/Private/SceneComponentDetails.cpp
@@ -251,7 +251,7 @@ void FSceneComponentDetails::MakeTransformDetails( IDetailLayoutBuilder& DetailB
 		for (int32 ComponentIndex = 0; bShouldShowTransform && ComponentIndex < SceneComponentObjects.Num(); ++ComponentIndex)
 		{
 			USceneComponent* SceneComponent = Cast<USceneComponent>(SceneComponentObjects[ComponentIndex].Get());
-			if (SceneComponent && SceneComponent->AttachParent == NULL && SceneComponent->GetOuter()->HasAnyFlags(RF_ClassDefaultObject))
+			if (SceneComponent && SceneComponent->GetAttachParent() == NULL && SceneComponent->GetOuter()->HasAnyFlags(RF_ClassDefaultObject))
 			{
 				bShouldShowTransform = false;
 			}
diff --git a/Engine/Source/Editor/DetailCustomizations/Private/StaticMeshComponentDetails.cpp b/Engine/Source/Editor/DetailCustomizations/Private/StaticMeshComponentDetails.cpp
--- a/Engine/Source/Editor/DetailCustomizations/Private/StaticMeshComponentDetails.cpp
+++ b/Engine/Source/Editor/DetailCustomizations/Private/StaticMeshComponentDetails.cpp
@@ -15,7 +15,46 @@ void FStaticMeshComponentDetails::CustomizeDetails( IDetailLayoutBuilder& Detail
 {
 	// Create a category so this is displayed early in the properties
 	DetailBuilder.EditCategory( "StaticMesh", FText::GetEmpty(), ECategoryPriority::Important);
+
+	DetailBuilder.GetObjectsBeingCustomized(ObjectsCustomized);
+	
+	TSharedPtr<IPropertyHandle> UseDefaultCollisionHandle = DetailBuilder.GetProperty(GET_MEMBER_NAME_CHECKED(UStaticMeshComponent, bUseDefaultCollision));
+
+	DetailBuilder.EditCategory(TEXT("Collision")).AddCustomRow(UseDefaultCollisionHandle->GetPropertyDisplayName())
+	.IsEnabled(TAttribute<bool>(this, &FStaticMeshComponentDetails::IsDefaultCollisionSupported))
+	.NameContent()
+	[
+		UseDefaultCollisionHandle->CreatePropertyNameWidget()
+	]
+	.ValueContent()
+	[
+		UseDefaultCollisionHandle->CreatePropertyValueWidget()
+	];
+
+	UseDefaultCollisionHandle->MarkHiddenByCustomization();
 }
 
+bool FStaticMeshComponentDetails::IsDefaultCollisionSupported() const
+{
+	// Whether to show the default collision property
+	bool bShowCollisionDefault = false;
+	for (TWeakObjectPtr<UObject> Obj : ObjectsCustomized)
+	{
+		if (UStaticMeshComponent* SMC = Cast<UStaticMeshComponent>(Obj.Get()))
+		{
+			if (SMC->SupportsDefaultCollision())
+			{
+				bShowCollisionDefault = true;
+			}
+			else
+			{
+				bShowCollisionDefault = false;
+				break;
+			}
+		}
+	}
+
+	return bShowCollisionDefault;
+}
 
 #undef LOCTEXT_NAMESPACE
diff --git a/Engine/Source/Editor/DetailCustomizations/Private/StaticMeshComponentDetails.h b/Engine/Source/Editor/DetailCustomizations/Private/StaticMeshComponentDetails.h
--- a/Engine/Source/Editor/DetailCustomizations/Private/StaticMeshComponentDetails.h
+++ b/Engine/Source/Editor/DetailCustomizations/Private/StaticMeshComponentDetails.h
@@ -16,6 +16,9 @@ private:
 	 * @return Whether or not override lightmap res is enabled
 	 */
 	bool IsOverrideLightmapResEnabled() const;
+
+	bool IsDefaultCollisionSupported() const;
 private:
 	TSharedPtr<IPropertyHandle> EnableOverrideLightmapRes;
+	TArray<TWeakObjectPtr<UObject>> ObjectsCustomized;
 };
diff --git a/Engine/Source/Editor/Documentation/Private/Documentation.cpp b/Engine/Source/Editor/Documentation/Private/Documentation.cpp
--- a/Engine/Source/Editor/Documentation/Private/Documentation.cpp
+++ b/Engine/Source/Editor/Documentation/Private/Documentation.cpp
@@ -68,7 +68,7 @@ bool FDocumentation::Open(const FString& Link, FDocumentationSourceInfo Source)
 		FText Message = LOCTEXT("OpeningURLMessage", "You are about to open an external URL. This will open your web browser. Do you want to proceed?");
 		FText URLDialog = LOCTEXT("OpeningURLTitle", "Open external link");
 
-		FSuppressableWarningDialog::FSetupInfo Info(Message, URLDialog, "SupressOpenURLWarning");
+		FSuppressableWarningDialog::FSetupInfo Info(Message, URLDialog, "SuppressOpenURLWarning");
 		Info.ConfirmText = LOCTEXT("OpenURL_yes", "Yes");
 		Info.CancelText = LOCTEXT("OpenURL_no", "No");
 		FSuppressableWarningDialog OpenURLWarning(Info);
diff --git a/Engine/Source/Editor/EditorStyle/Private/SlateEditorStyle.cpp b/Engine/Source/Editor/EditorStyle/Private/SlateEditorStyle.cpp
--- a/Engine/Source/Editor/EditorStyle/Private/SlateEditorStyle.cpp
+++ b/Engine/Source/Editor/EditorStyle/Private/SlateEditorStyle.cpp
@@ -4311,7 +4311,12 @@ void FSlateEditorStyle::FStyle::SetupLevelEditorStyle()
 		Set( "EditorViewport.TexelFactorAccuracyMode", new IMAGE_BRUSH( "Icons/icon_ViewMode_TexelFactorAccuracy_16px", Icon16x16 ) );
 		Set( "EditorViewport.StationaryLightOverlapMode", new IMAGE_BRUSH( "Icons/icon_ViewMode_StationaryLightOverlap_16px", Icon16x16 ) );
 		Set( "EditorViewport.LightmapDensityMode", new IMAGE_BRUSH( "Icons/icon_ViewMode_LightmapDensity_16px", Icon16x16 ) );
+
 		Set( "EditorViewport.LODColorationMode", new IMAGE_BRUSH("Icons/icon_ViewMode_LODColoration_16px", Icon16x16) );
+		Set( "EditorViewport.HLODColorationMode", new IMAGE_BRUSH("Icons/icon_ViewMode_LODColoration_16px", Icon16x16));
+		Set( "EditorViewport.GroupLODColorationMode", new IMAGE_BRUSH("Icons/icon_ViewMode_LODColoration_16px", Icon16x16));
+
+
 		Set( "EditorViewport.VisualizeGBufferMode", new IMAGE_BRUSH("Icons/icon_ViewMode_VisualisationGBuffer_16px", Icon16x16) );
 		Set( "EditorViewport.ReflectionOverrideMode", new IMAGE_BRUSH( "Icons/icon_ViewMode_ReflectionOverride_16px", Icon16x16 ) );
 		Set( "EditorViewport.VisualizeBufferMode", new IMAGE_BRUSH( "Icons/icon_ViewMode_VisualisationGBuffer_16px", Icon16x16 ) );
diff --git a/Engine/Source/Editor/HierarchicalLODOutliner/Private/HLODOutliner.cpp b/Engine/Source/Editor/HierarchicalLODOutliner/Private/HLODOutliner.cpp
--- a/Engine/Source/Editor/HierarchicalLODOutliner/Private/HLODOutliner.cpp
+++ b/Engine/Source/Editor/HierarchicalLODOutliner/Private/HLODOutliner.cpp
@@ -25,13 +25,16 @@
 #include "HLODSelectionActor.h"
 #include "TreeItemID.h"
 
+#include "MessageLog.h"
+#include "UObjectToken.h"
+
 #include "HierarchicalLODUtilities.h"
 
 #define LOCTEXT_NAMESPACE "HLODOutliner"
 
 namespace HLODOutliner
 {
-	SHLODOutliner::SHLODOutliner()
+	SHLODOutliner::SHLODOutliner()		
 	{
 		bNeedsRefresh = true;
 		CurrentWorld = nullptr;
@@ -132,6 +135,7 @@ namespace HLODOutliner
 					.HAlign(HAlign_Center)
 					.Text(LOCTEXT("GenerateClusters", "Generate Clusters"))
 					.OnClicked(this, &SHLODOutliner::HandlePreviewHLODs)
+					.ToolTipText(LOCTEXT("GenerateClusterToolTip", "Generates Clusters for Meshes in the Level"))
 				]
 
 				+ SHorizontalBox::Slot()
@@ -143,6 +147,7 @@ namespace HLODOutliner
 					.HAlign(HAlign_Center)
 					.Text(LOCTEXT("DeleteClusters", "Delete Clusters"))
 					.OnClicked(this, &SHLODOutliner::HandleDeleteHLODs)
+					.ToolTipText(LOCTEXT("DeleteClusterToolTip", "Deletes all Clusters in the Level"))
 				]
 
 				+ SHorizontalBox::Slot()
@@ -152,11 +157,11 @@ namespace HLODOutliner
 				[
 					SNew(SButton)
 					.HAlign(HAlign_Center)
-					.Text(LOCTEXT("BuildMeshes", "Build LOD Meshes for Clusters"))
+					.Text(LOCTEXT("BuildMeshes", "Generate Proxy Meshes"))
 					.OnClicked(this, &SHLODOutliner::HandleBuildLODActors)
-				]				
-			];
-		
+					.ToolTipText(LOCTEXT("GenerateProxyMeshesToolTip", "Generates Proxy Mesh for each Clusters in the Level"))
+				]
+			];	
 	}
 
 	TSharedRef<SWidget> SHLODOutliner::CreateTreeviewWidget()
@@ -175,12 +180,24 @@ namespace HLODOutliner
 				SNew(SHeaderRow)
 				+ SHeaderRow::Column("SceneActorName")
 				.DefaultLabel(LOCTEXT("SceneActorName", "Scene Actor Name"))
-				.FillWidth(0.5f)				
-				+ SHeaderRow::Column("TriangleCount")
-				.DefaultLabel(LOCTEXT("TriangleCount", "Number of Triangles"))
-				.DefaultTooltip(LOCTEXT("TriangleCountToolTip", "Number of Triangles in a LOD Mesh"))
-				.FillWidth(0.5f)				
-				);
+				.FillWidth(0.3f)				
+				+ SHeaderRow::Column("RawTriangleCount")
+				.DefaultLabel(LOCTEXT("RawTriangleCount", "Original Triangle Count"))
+				.DefaultTooltip(LOCTEXT("RawTriangleCountToolTip", "Original Number of Triangles in a LOD Mesh"))
+				.FillWidth(0.2f)
+				+ SHeaderRow::Column("ReducedTriangleCount")
+				.DefaultLabel(LOCTEXT("ReducedTriangleCount", "Reduced Triangle Count"))
+				.DefaultTooltip(LOCTEXT("ReducedTriangleCountToolTip", "Reduced Number of Triangles in a LOD Mesh"))
+				.FillWidth(0.2f)
+				+ SHeaderRow::Column("ReductionPercentage")
+				.DefaultLabel(LOCTEXT("ReductionPercentage", "% Reduction"))
+				.DefaultTooltip(LOCTEXT("ReductionPercentageToolTip", "Percentage of Triangle Reduction in a LOD Mesh"))
+				.FillWidth(0.1f)
+				+ SHeaderRow::Column("Level")
+				.DefaultLabel(LOCTEXT("Level", "Level"))
+				.DefaultTooltip(LOCTEXT("LevelToolTip", "Persistent Level of a LOD Mesh"))
+				.FillWidth(0.2f)
+			);
 	}
 
 	TSharedRef<SWidget> SHLODOutliner::CreateForcedViewSlider()
@@ -354,6 +371,8 @@ namespace HLODOutliner
 			CurrentWorld->HierarchicalLODBuilder->Build();
 		}
 
+		FMessageLog("HLODResults").Open();
+
 		FullRefresh();
 		return FReply::Handled();
 	}
@@ -379,6 +398,9 @@ namespace HLODOutliner
 		{
 			CurrentWorld->HierarchicalLODBuilder->PreviewBuild();
 		}
+
+		FMessageLog("HLODResults").Open();
+
 		FullRefresh();
 		return FReply::Handled();
 	}
@@ -401,6 +423,8 @@ namespace HLODOutliner
 			CurrentWorld->HierarchicalLODBuilder->BuildMeshesForLODActors();
 		}
 
+		FMessageLog("HLODResults").Open();
+
 		ResetLODLevelForcing();
 
 		FullRefresh();
@@ -416,10 +440,11 @@ namespace HLODOutliner
 
 	END_SLATE_FUNCTION_BUILD_OPTIMIZATION
 
-	void SHLODOutliner::RegisterDelegates()
+		void SHLODOutliner::RegisterDelegates()
 	{
 		FEditorDelegates::MapChange.AddSP(this, &SHLODOutliner::OnMapChange);
 		FEditorDelegates::NewCurrentLevel.AddSP(this, &SHLODOutliner::OnNewCurrentLevel);
+		FEditorDelegates::OnMapOpened.AddSP(this, &SHLODOutliner::OnMapLoaded);
 		FWorldDelegates::LevelAddedToWorld.AddSP(this, &SHLODOutliner::OnLevelAdded);
 		FWorldDelegates::LevelRemovedFromWorld.AddSP(this, &SHLODOutliner::OnLevelRemoved);
 		GEngine->OnLevelActorListChanged().AddSP(this, &SHLODOutliner::FullRefresh);
@@ -435,11 +460,12 @@ namespace HLODOutliner
 		USelection::SelectObjectEvent.AddRaw(this, &SHLODOutliner::OnLevelSelectionChanged);
 
 		// HLOD related events
-		GEngine->OnHLODActorMoved().AddSP(this, &SHLODOutliner::OnHLODActorMovedEvent);
-		GEngine->OnHLODActorAdded().AddSP(this, &SHLODOutliner::OnHLODActorAddedEvent);
-		GEngine->OnHLODActorMarkedDirty().AddSP(this, &SHLODOutliner::OnHLODActorMarkedDirtyEvent);
-		GEngine->OnHLODTransitionScreenSizeChanged().AddSP(this, &SHLODOutliner::OnHLODTransitionScreenSizeChangedEvent);
-		GEngine->OnHLODLevelsArrayChanged().AddSP(this, &SHLODOutliner::OnHLODLevelsArrayChangedEvent);
+		GEditor->OnHLODActorMoved().AddSP(this, &SHLODOutliner::OnHLODActorMovedEvent);
+		GEditor->OnHLODActorAdded().AddSP(this, &SHLODOutliner::OnHLODActorAddedEvent);
+		GEditor->OnHLODActorMarkedDirty().AddSP(this, &SHLODOutliner::OnHLODActorMarkedDirtyEvent);
+		GEditor->OnHLODTransitionScreenSizeChanged().AddSP(this, &SHLODOutliner::OnHLODTransitionScreenSizeChangedEvent);
+		GEditor->OnHLODLevelsArrayChanged().AddSP(this, &SHLODOutliner::OnHLODLevelsArrayChangedEvent);
+		GEditor->OnHLODActorRemovedFromCluster().AddSP(this, &SHLODOutliner::OnHLODActorRemovedFromClusterEvent);
 
 		// Register to update when an undo/redo operation has been called to update our list of actors
 		GEditor->RegisterForUndo(this);
@@ -449,25 +475,33 @@ namespace HLODOutliner
 	{
 		FEditorDelegates::MapChange.RemoveAll(this);
 		FEditorDelegates::NewCurrentLevel.RemoveAll(this);
+		FEditorDelegates::OnMapOpened.RemoveAll(this);
+
 		FWorldDelegates::LevelAddedToWorld.RemoveAll(this);
 		FWorldDelegates::LevelRemovedFromWorld.RemoveAll(this);
-		GEngine->OnLevelActorListChanged().RemoveAll(this);
-		GEngine->OnLevelActorAdded().RemoveAll(this);
-		GEngine->OnLevelActorDeleted().RemoveAll(this);
-		GEngine->OnActorMoved().RemoveAll(this);
-
 		FCoreDelegates::OnActorLabelChanged.RemoveAll(this);
-
 		USelection::SelectionChangedEvent.RemoveAll(this);
 		USelection::SelectObjectEvent.RemoveAll(this);
 
-		GEngine->OnHLODActorMoved().RemoveAll(this);
-		GEngine->OnHLODActorAdded().RemoveAll(this);
-		GEngine->OnHLODActorMarkedDirty().RemoveAll(this);
-		GEngine->OnHLODLevelsArrayChanged().RemoveAll(this);
+		if (GEngine)
+		{
+			GEngine->OnLevelActorListChanged().RemoveAll(this);
+			GEngine->OnLevelActorAdded().RemoveAll(this);
+			GEngine->OnLevelActorDeleted().RemoveAll(this);
+			GEngine->OnActorMoved().RemoveAll(this);
+		}	
+		
+		if (GEditor)
+		{
+			GEditor->OnHLODActorMoved().RemoveAll(this);
+			GEditor->OnHLODActorAdded().RemoveAll(this);
+			GEditor->OnHLODActorMarkedDirty().RemoveAll(this);
+			GEditor->OnHLODLevelsArrayChanged().RemoveAll(this);
+			GEditor->OnHLODActorRemovedFromCluster().RemoveAll(this);
 
-		// Deregister for Undo callbacks
-		GEditor->UnregisterForUndo(this);
+			// Deregister for Undo callbacks
+			GEditor->UnregisterForUndo(this);
+		}
 	}
 
 	void SHLODOutliner::ForceViewLODActor()
@@ -637,9 +671,12 @@ namespace HLODOutliner
 				}
 			}
 
-			ResetLODLevelForcing();
+			ResetLODLevelForcing();			
 			FullRefresh();			
 		}
+		
+		// Show message log if there was an HLOD message
+		FMessageLog("HLODResults").Open();		
 	}
 
 	void SHLODOutliner::RebuildLODActor()
@@ -670,6 +707,9 @@ namespace HLODOutliner
 			ResetLODLevelForcing();
 			FullRefresh();
 		}
+
+		// Show message log if there was an HLOD message
+		FMessageLog("HLODResults").Open();
 	}
 
 	void SHLODOutliner::SelectLODActor()
@@ -683,6 +723,8 @@ namespace HLODOutliner
 			EmptySelection();
 			StartSelection();
 
+
+			bool bChanged = false;
 			// Loop over all selected items (context menu can't be called with multiple items selected that aren't of the same types)
 			for (auto SelectedItem : SelectedItems)
 			{
@@ -690,19 +732,18 @@ namespace HLODOutliner
 
 				if (ActorItem->LODActor.IsValid())
 				{				
-					SelectActorInViewport(ActorItem->LODActor.Get(), 0);					
+					SelectActorInViewport(ActorItem->LODActor.Get(), 0);			
+					bChanged = true;
 				}
 			}
 			
 			// Done selecting actors
-			EndSelection();
+			EndSelection(bChanged);
 		}
 	}
 
 	void SHLODOutliner::DeleteCluster()
 	{
-		const FScopedTransaction Transaction(LOCTEXT("UndoAction_DestroyLODActor", "Delete LOD Actor"));
-
 		// This call came from a context menu
 		auto SelectedItems = TreeView->GetSelectedItems();
 		
@@ -710,28 +751,11 @@ namespace HLODOutliner
 		for (auto SelectedItem : SelectedItems)
 		{
 			FLODActorItem* ActorItem = (FLODActorItem*)(SelectedItem.Get());
-
-			ALODActor* LODActor = ActorItem->LODActor.Get();
-			ALODActor* ParentActor = FHierarchicalLODUtilities::GetParentLODActor(LODActor);
-
-			LODActor->Modify();
-
-			if (ParentActor)
-			{
-				ParentActor->Modify();
-			}
-
-			FHierarchicalLODUtilities::DeleteLODActor(LODActor);
-			CurrentWorld->DestroyActor(LODActor);
-
-			if (ParentActor && !ParentActor->HasValidSubActors())
-			{
-				DestroyLODActor(ParentActor);
-			}
+			ALODActor* LODActor = ActorItem->LODActor.Get();			
+			FHierarchicalLODUtilities::DestroyLODActor(LODActor);
 		}
 
 		ResetLODLevelForcing();
-
 		FullRefresh();
 	}
 
@@ -753,17 +777,11 @@ namespace HLODOutliner
 				ITreeItem::TreeItemType Type = Parent->GetTreeItemType();
 				if (Type == ITreeItem::HierarchicalLODActor)
 				{
-					FLODActorItem* ParentLODActorItem = (FLODActorItem*)(Parent.Get());
-					ParentLODActorItem->LODActor->Modify();
-					ActorItem->StaticMeshActor->Modify();
-					ParentLODActorItem->LODActor->RemoveSubActor(ActorItem->StaticMeshActor.Get());
-
-					PendingActions.Emplace(FOutlinerAction::RemoveItem, SelectedItem);
-
-					if (!ParentLODActorItem->LODActor->HasValidSubActors())
+					AActor* Actor = Cast<AActor>(ActorItem->StaticMeshActor.Get());
+					
+					if (FHierarchicalLODUtilities::RemoveActorFromCluster(Actor))
 					{
-						DestroyLODActor(ParentLODActorItem->LODActor.Get());
-						PendingActions.Emplace(FOutlinerAction::RemoveItem, Parent);
+						PendingActions.Emplace(FOutlinerAction::RemoveItem, SelectedItem);
 					}
 				}
 			}
@@ -772,7 +790,6 @@ namespace HLODOutliner
 
 	void SHLODOutliner::ExcludeFromClusterGeneration()
 	{
-		const FScopedTransaction Transaction(LOCTEXT("UndoAction_ExcludeStaticMeshActorFromClusterGeneration", "Excluded StaticMeshActor From Cluster Generation"));
 		// This call came from a context menu
 		auto SelectedItems = TreeView->GetSelectedItems();
 
@@ -780,9 +797,7 @@ namespace HLODOutliner
 		for (auto SelectedItem : SelectedItems)
 		{
 			FStaticMeshActorItem* ActorItem = (FStaticMeshActorItem*)(SelectedItem.Get());
-			ActorItem->StaticMeshActor->Modify();
-			ActorItem->StaticMeshActor->bEnableAutoLODGeneration = false;
-			RemoveStaticMeshActorFromCluster();
+			FHierarchicalLODUtilities::ExcludeActorFromClusterGeneration(ActorItem->StaticMeshActor.Get());
 		}		
 	}
 
@@ -796,25 +811,19 @@ namespace HLODOutliner
 			// Loop over all selected items (context menu can't be called with multiple items selected that aren't of the same types)
 			for (auto SelectedItem : SelectedItems)
 			{
-				const FScopedTransaction Transaction(LOCTEXT("UndoAction_RemoveLODActorFromCluster", "Removed LOD Actor From Cluster"));
 				FLODActorItem* ActorItem = (FLODActorItem*)(SelectedItem.Get());
 				auto Parent = ActorItem->GetParent();
 
 				ITreeItem::TreeItemType Type = Parent->GetTreeItemType();
 				if (Type == ITreeItem::HierarchicalLODActor)
 				{
-					FLODActorItem* ParentLODActorItem = (FLODActorItem*)(Parent.Get());
-					ParentLODActorItem->LODActor->Modify();
-					ActorItem->LODActor->Modify();
-					ParentLODActorItem->LODActor->RemoveSubActor(ActorItem->LODActor.Get());
-
-					PendingActions.Emplace(FOutlinerAction::RemoveItem, SelectedItem);
-
-					if (!ParentLODActorItem->LODActor->HasValidSubActors())
+					AActor* Actor = Cast<AActor>(ActorItem->LODActor.Get());
+					checkf(Actor != nullptr, TEXT("Invalid actor in tree view"));
+					
+					if (FHierarchicalLODUtilities::RemoveActorFromCluster(Actor))
 					{
-						DestroyLODActor(ParentLODActorItem->LODActor.Get());
-						PendingActions.Emplace(FOutlinerAction::RemoveItem, Parent);
-					}
+						PendingActions.Emplace(FOutlinerAction::RemoveItem, SelectedItem);
+					}					
 				}
 			}
 		}
@@ -836,24 +845,11 @@ namespace HLODOutliner
 			FLODActorItem* ActorItem = (FLODActorItem*)(SelectedItem.Get());
 
 			ALODActor* LODActor = ActorItem->LODActor.Get();
-			SelectLODActorAndContainedActorsInViewport(LODActor, 0);
+			SelectContainedActorsInViewport(LODActor, 0);
 		}
 
 		// Done selecting actors
-		EndSelection();
-	}
-
-	void SHLODOutliner::DestroyLODActor(ALODActor* InActor)
-	{		
-		ALODActor* ParentActor = FHierarchicalLODUtilities::GetParentLODActor(InActor);
-
-		FHierarchicalLODUtilities::DeleteLODActor(InActor);
-		CurrentWorld->DestroyActor(InActor);
-
-		if (ParentActor && !ParentActor->HasValidSubActors())
-		{
-			DestroyLODActor(ParentActor);
-		}
+		EndSelection(true);
 	}
 
 	void SHLODOutliner::UpdateDrawDistancesForLODLevel(const uint32 LODLevelIndex)
@@ -922,10 +918,13 @@ namespace HLODOutliner
 		TArray<FTreeItemPtr> NewSelectedNodes = TreeView->GetSelectedItems();
 
 		EmptySelection();
-		TreeView->ClearSelection();
+		//TreeView->ClearSelection();
 
 		// Loop over previously retrieve lsit of selected nodes
 		StartSelection();
+
+		bool bChanged = false;
+
 		for (FTreeItemPtr SelectedItem : NewSelectedNodes)
 		{
 			if (SelectedItem.IsValid())
@@ -933,38 +932,34 @@ namespace HLODOutliner
 				ITreeItem::TreeItemType Type = SelectedItem->GetTreeItemType();
 				switch (Type)
 				{
-				case ITreeItem::HierarchicalLODLevel:
-				{
-					FLODLevelItem* LevelItem = (FLODLevelItem*)(SelectedItem.Get());
-					const TArray<TWeakPtr<ITreeItem>>& Children = LevelItem->GetChildren();
-					for (auto& WeakChild : Children)
+					case ITreeItem::HierarchicalLODLevel:
 					{
-						auto Child = WeakChild.Pin();
-						check(Child.IsValid());
-						FLODActorItem* ActorItem = (FLODActorItem*)(Child.Get());
-						SelectLODActorAndContainedActorsInViewport(ActorItem->LODActor.Get(), 0);
+						// No functionality for select HLOD level items
+						break;
 					}
 
-					break;
-				}
-
-				case ITreeItem::HierarchicalLODActor:
-				{
-					FLODActorItem* ActorItem = (FLODActorItem*)(SelectedItem.Get());
-					SelectActorInViewport(ActorItem->LODActor.Get(), 0);
-					break;
-				}
+					case ITreeItem::HierarchicalLODActor:
+					{
+						FLODActorItem* ActorItem = (FLODActorItem*)(SelectedItem.Get());
+						if (ActorItem->GetParent()->GetTreeItemType() == ITreeItem::HierarchicalLODLevel)
+						{
+							SelectActorInViewport(ActorItem->LODActor.Get(), 0);
+							bChanged = true;
+						}
+						break;
+					}
 
-				case ITreeItem::StaticMeshActor:
-				{
-					FStaticMeshActorItem* StaticMeshActorItem = (FStaticMeshActorItem*)(SelectedItem.Get());
-					SelectActorInViewport(StaticMeshActorItem->StaticMeshActor.Get(), 0);
-					break;
-				}
+					case ITreeItem::StaticMeshActor:
+					{
+						FStaticMeshActorItem* StaticMeshActorItem = (FStaticMeshActorItem*)(SelectedItem.Get());
+						SelectActorInViewport(StaticMeshActorItem->StaticMeshActor.Get(), 0);
+						bChanged = true;
+						break;
+					}
 				}
 			}
 		}
-		EndSelection();
+		EndSelection(bChanged);
 
 		SelectedNodes = TreeView->GetSelectedItems();
 	}
@@ -1097,7 +1092,7 @@ namespace HLODOutliner
 	void SHLODOutliner::SelectLODActorAndContainedActorsInViewport(ALODActor* LODActor, const uint32 SelectionDepth)
 	{
 		TArray<AActor*> SubActors;
-		ExtractStaticMeshActorsFromLODActor(LODActor, SubActors);
+		FHierarchicalLODUtilities::ExtractStaticMeshActorsFromLODActor(LODActor, SubActors);
 		for (AActor* SubActor : SubActors)
 		{
 			SelectActorInViewport(SubActor, SelectionDepth + 1);
@@ -1106,6 +1101,16 @@ namespace HLODOutliner
 		SelectActorInViewport(LODActor, SelectionDepth);
 	}
 
+	void SHLODOutliner::SelectContainedActorsInViewport(ALODActor* LODActor, const uint32 SelectionDepth /*= 0*/)
+	{
+		TArray<AActor*> SubActors;
+		FHierarchicalLODUtilities::ExtractStaticMeshActorsFromLODActor(LODActor, SubActors);
+		for (AActor* SubActor : SubActors)
+		{
+			SelectActorInViewport(SubActor, SelectionDepth + 1);
+		}
+	}
+
 	UDrawSphereComponent* SHLODOutliner::CreateBoundingSphereForActor(AActor* Actor)
 	{
 		if (CurrentWorld)
@@ -1132,31 +1137,16 @@ namespace HLODOutliner
 		return nullptr;
 	}
 
-	void SHLODOutliner::ExtractStaticMeshActorsFromLODActor(ALODActor* LODActor, TArray<AActor*> &InOutActors)
-	{
-		for (auto ChildActor : LODActor->SubActors)
-		{
-			if (ChildActor)
-			{
-				TArray<AActor*> ChildActors;
-				if (ChildActor->IsA<ALODActor>())
-				{
-					ExtractStaticMeshActorsFromLODActor(Cast<ALODActor>(ChildActor), ChildActors);
-				}
-
-				ChildActors.Push(ChildActor);
-				InOutActors.Append(ChildActors);
-			}
-		}
-	}
-
-	void SHLODOutliner::EndSelection()
+	void SHLODOutliner::EndSelection(const bool bChange)
 	{
 		// Commit selection changes
 		GEditor->GetSelectedActors()->EndBatchSelectOperation();
 
-		// Fire selection changed event
-		GEditor->NoteSelectionChange();
+		if (bChange)
+		{
+			// Fire selection changed event
+			GEditor->NoteSelectionChange();
+		}
 	}
 
 	void SHLODOutliner::OnLevelSelectionChanged(UObject* Obj)
@@ -1176,7 +1166,7 @@ namespace HLODOutliner
 					auto Item = TreeItemsMap.Find(Actor);
 					if (Item)
 					{
-						SelectItemInTree(*Item);						
+						SelectItemInTree(*Item);
 
 						if (Item->Get()->GetTreeItemType() == ITreeItem::StaticMeshActor)
 						{
@@ -1210,7 +1200,7 @@ namespace HLODOutliner
 			{
 				DestroySelectionActors();
 			}
-		}		
+		}
 	}
 
 	void SHLODOutliner::OnLevelAdded(ULevel* InLevel, UWorld* InWorld)
@@ -1230,31 +1220,35 @@ namespace HLODOutliner
 		if (!InActor->IsA<AHLODSelectionActor>() && !InActor->IsA<AWorldSettings>())
 		{
 			FullRefresh();
-		}
+		}	
 	}
 
 	void SHLODOutliner::OnLevelActorsRemoved(AActor* InActor)
 	{
 		if (!InActor->IsA<AHLODSelectionActor>() && !InActor->IsA<AWorldSettings>())
-		{
+		{			
 			// Remove InActor from LOD actor which contains it
 			for (TArray<ALODActor*>& ActorArray : LODLevelActors)
 			{				
 				for (ALODActor* Actor : ActorArray)
 				{
+					// Check if actor is not null due to Destroy Actor
 					if (Actor)
 					{
 						checkf(Actor->IsValidLowLevel(), TEXT("Invalid LODActor found in ActorArray"));
 						Actor->CleanSubActorArray();
-						if (Actor->RemoveSubActor(InActor))
+						const bool bRemovedSubActor = Actor->RemoveSubActor(InActor);
+
+						if (bRemovedSubActor && Actor->SubActors.Num() == 0)
 						{
-							break;
+							FHierarchicalLODUtilities::DestroyCluster(Actor);
+							Actor->GetWorld()->DestroyActor(Actor);						
 						}
 					}
 				}
 			}
 			FullRefresh();
-		}
+		} 
 	}
 	
 	void SHLODOutliner::OnActorLabelChanged(AActor* ChangedActor)
@@ -1262,7 +1256,7 @@ namespace HLODOutliner
 		if (!ChangedActor->IsA<AHLODSelectionActor>())
 		{
 			FullRefresh();
-		}
+		}	
 	}
 
 	void SHLODOutliner::OnMapChange(uint32 MapFlags)
@@ -1272,6 +1266,12 @@ namespace HLODOutliner
 	}
 
 	void SHLODOutliner::OnNewCurrentLevel()
+	{	
+		ResetCachedData();
+		FullRefresh();
+	}
+
+	void SHLODOutliner::OnMapLoaded(const FString& Filename, bool bAsTemplate)
 	{
 		ResetCachedData();
 		FullRefresh();
@@ -1294,7 +1294,7 @@ namespace HLODOutliner
 					FLODActorItem* ParentLODActorItem = (FLODActorItem*)CurrentParent.Get();
 					if (!ParentLODActorItem->LODActor->HasValidSubActors())
 					{
-						DestroyLODActor(ParentLODActorItem->LODActor.Get());
+						FHierarchicalLODUtilities::DestroyLODActor(ParentLODActorItem->LODActor.Get());
 						PendingActions.Emplace(FOutlinerAction::RemoveItem, CurrentParent);
 					}
 				}
@@ -1319,23 +1319,57 @@ namespace HLODOutliner
 
 	void SHLODOutliner::OnHLODActorAddedEvent(const AActor* InActor, const AActor* ParentActor)
 	{
-		FullRefresh();
-	}
+		checkf(InActor != nullptr, TEXT("Invalid InActor found"));
+		checkf(ParentActor != nullptr, TEXT("Invalid ParentActor found"));
 
-	void SHLODOutliner::OnHLODActorMarkedDirtyEvent(ALODActor* InActor)
-	{
-		if (InActor->GetStaticMeshComponent()->StaticMesh)
+		FTreeItemPtr* ParentItem = TreeItemsMap.Find(ParentActor);
+		if (ParentItem->IsValid())
 		{
-			InActor->GetStaticMeshComponent()->StaticMesh = nullptr;
+			const ALODActor* ParentLODActor = Cast<ALODActor>(ParentActor);
+			
+			FTreeItemPtr* ChildItemPtr = TreeItemsMap.Find(InActor);
+			if (ChildItemPtr)
+			{
+				if (!InActor->IsA<ALODActor>())
+				{
+					PendingActions.Emplace(FOutlinerAction::MoveItem, *ChildItemPtr, *ParentItem);
+				}
+				else
+				{
+					FullRefresh();
+
+					// TODO, handle duplicate entries of same actor in tree view
+					/*// Add child item for the new LOD mesh actor
+					const ALODActor* LODActor = Cast<ALODActor>(InActor);
+					FTreeItemRef ChildItem = MakeShareable(new FLODActorItem(LODActor));
+					AllNodes.Add(ChildItem->AsShared());
+					(*ParentItem)->AddChild(ChildItem);
+
+					TreeView->RequestTreeRefresh();*/
+				}
+			}
+			else
+			{
+				// Add child item for the new static mesh actor
+				FTreeItemRef ChildItem = MakeShareable(new FStaticMeshActorItem(InActor));
+				AllNodes.Add(ChildItem->AsShared());
+				PendingActions.Emplace(FOutlinerAction::AddItem, ChildItem, *ParentItem);
+			}			
+
+			// Set build flags according to whether or not this LOD actor is dirty 
+			LODLevelBuildFlags[ParentLODActor->LODLevel - 1] &= !ParentLODActor->IsDirty();
 		}
+	}
 
-		FullRefresh(); 
+	void SHLODOutliner::OnHLODActorMarkedDirtyEvent(ALODActor* InActor)
+	{	
+		bool check = true;
 	}
 
 	void SHLODOutliner::OnHLODTransitionScreenSizeChangedEvent()
 	{
 		if (CurrentWorld)
-		{
+		{			
 			int32 MaxLODLevel = FMath::Min(CurrentWorldSettings->HierarchicalLODSetup.Num(), LODLevelTransitionScreenSizes.Num());
 			for (int32 LODLevelIndex = 0; LODLevelIndex < MaxLODLevel; ++LODLevelIndex)
 			{
@@ -1356,11 +1390,23 @@ namespace HLODOutliner
 		}
 	}
 
+	void SHLODOutliner::OnHLODActorRemovedFromClusterEvent(const AActor* InActor, const AActor* ParentActor)
+	{
+		FTreeItemPtr* TreeItem = TreeItemsMap.Find(InActor);
+		FTreeItemPtr* ParentItem = TreeItemsMap.Find(ParentActor);
+		if (TreeItem->IsValid() && ParentItem->IsValid())
+		{
+			checkf((*TreeItem)->GetTreeItemType() == ITreeItem::StaticMeshActor, TEXT("Incorrect InActor"));
+			checkf((*ParentItem)->GetTreeItemType() == ITreeItem::HierarchicalLODActor, TEXT("Incorrect ParentActor"));
+			PendingActions.Emplace(FOutlinerAction::RemoveItem, *TreeItem);
+		}
+	}
+
 	void SHLODOutliner::FullRefresh()
 	{		
 		bNeedsRefresh = true;
 	}
-
+	
 	const bool SHLODOutliner::UpdateCurrentWorldAndSettings()
 	{
 		CurrentWorld = nullptr;
@@ -1382,11 +1428,11 @@ namespace HLODOutliner
 			// Retrieve current world settings
 			CurrentWorldSettings = CurrentWorld->GetWorldSettings();
 			checkf(CurrentWorldSettings != nullptr, TEXT("CurrentWorld (%s) does not contain a valid WorldSettings actor"), *CurrentWorld->GetName());
-
+						
 			// Update settings view
 			SettingsView->SetObject(CurrentWorldSettings);
 		}
-
+		
 
 		return (CurrentWorld != nullptr);
 	}
@@ -1401,7 +1447,7 @@ namespace HLODOutliner
 		if (CurrentWorldSettings)
 		{
 			// Iterate over all LOD levels (Number retrieved from world settings) and add Treeview items for them
-			const uint32 LODLevels = CurrentWorldSettings->HierarchicalLODSetup.Num();
+			const uint32 LODLevels = CurrentWorldSettings->HierarchicalLODSetup.Num();			
 			for (uint32 LODLevelIndex = 0; LODLevelIndex < LODLevels; ++LODLevelIndex)
 			{
 				FTreeItemRef LevelItem = MakeShareable(new FLODLevelItem(LODLevelIndex));
@@ -1415,7 +1461,7 @@ namespace HLODOutliner
 
 				// Initialize lod level actors/screen size and build flag
 				LODLevelBuildFlags.Add(true);
-				LODLevelActors.AddDefaulted();
+				LODLevelActors.AddDefaulted();					
 				LODLevelTransitionScreenSizes.Add(CurrentWorldSettings->HierarchicalLODSetup[LODLevelIndex].TransitionScreenSize);
 
 				TreeItemsMap.Add(LevelItem->GetID(), LevelItem);
@@ -1435,11 +1481,17 @@ namespace HLODOutliner
 						// Only handling LODActors
 						if (Actor)
 						{
-							ALODActor* LODActor = Cast<ALODActor>(Actor);
+							ALODActor* LODActor = Cast<ALODActor>(Actor);							
 							// Add LOD Actor item to the treeview
 							if (LODActor)
 							{
 								checkf((LODActor->LODLevel - 1) < LevelNodes.Num(), TEXT("LODActor (%s) found with LODLevel (%i) that is out of current WorldSettings range (%i)"), *LODActor->GetName(), LODActor->LODLevel - 1, LevelNodes.Num());
+								
+								// This is to prevent issues with the sub actor array due to deleted scene actors while the HLOD outliner was closed
+								LODActor->CleanSubActorArray();
+
+								// Set LOD parents here TODO remove if not needed anymore QQ
+								LODActor->UpdateSubActorLODParents();
 
 								FTreeItemRef Item = MakeShareable(new FLODActorItem(LODActor));
 								AllNodes.Add(Item->AsShared());
@@ -1471,7 +1523,7 @@ namespace HLODOutliner
 						}
 					}
 				}
-			}
+			}	
 
 			// Take empty LOD levels into account for the build flags
 			for (uint32 LODLevelIndex = 0; LODLevelIndex < LODLevels; ++LODLevelIndex)
@@ -1485,7 +1537,7 @@ namespace HLODOutliner
 
 		// Request treeview UI item to refresh
 		TreeView->RequestTreeRefresh();
-
+		
 		// Just finished refreshing
 		bNeedsRefresh = false;
 	}
@@ -1543,11 +1595,6 @@ namespace HLODOutliner
 	{
 		const auto ItemID = InItem->GetID();
 
-		if (TreeItemsMap.Find(ItemID))
-		{
-			return false;
-		}
-
 		TreeItemsMap.Add(ItemID, InItem);
 
 		if (InParentItem.Get())
@@ -1594,14 +1641,21 @@ namespace HLODOutliner
 		while (Parent.IsValid() && !Parent->bIsExpanded)
 		{
 			Parent->bIsExpanded = true;
+			TreeView->SetItemExpansion(Parent, true);
 			Parent = InItem->GetParent();
 		}
 		TreeView->SetItemSelection(InItem, true);
 
 		TreeView->RequestTreeRefresh();
 	}
+	
+	FReply SHLODOutliner::RetrieveActors()
+	{
+		bNeedsRefresh = true;
+		return FReply::Handled();
+	}
 
-	bool SHLODOutliner::IsHLODEnabledInWorldSettings()
+	bool SHLODOutliner::IsHLODEnabledInWorldSettings() const
 	{
 		bool bHLODEnabled = false;
 		if (CurrentWorldSettings != nullptr)
@@ -1612,11 +1666,6 @@ namespace HLODOutliner
 		return bHLODEnabled;
 	}
 
-	FReply SHLODOutliner::RetrieveActors()
-	{
-		bNeedsRefresh = true;
-		return FReply::Handled();
-	}
 
 };
 
diff --git a/Engine/Source/Editor/HierarchicalLODOutliner/Private/HLODOutliner.h b/Engine/Source/Editor/HierarchicalLODOutliner/Private/HLODOutliner.h
--- a/Engine/Source/Editor/HierarchicalLODOutliner/Private/HLODOutliner.h
+++ b/Engine/Source/Editor/HierarchicalLODOutliner/Private/HLODOutliner.h
@@ -202,13 +202,6 @@ namespace HLODOutliner
 		/** Removes the given LODActor from its parent's (ALODActor) sub-actors array */
 		void RemoveLODActorFromCluster();
 		
-		/**
-		* Destroys an LODActor instance
-		*
-		* @param InActor - ALODActor to destroy
-		*/
-		void DestroyLODActor(ALODActor* InActor);
-
 		/**
 		* Updates the DrawDistance value for all the LODActors with the given LODLevelIndex
 		*
@@ -289,13 +282,21 @@ namespace HLODOutliner
 		void SelectActorInViewport(AActor* Actor, const uint32 SelectionDepth = 0);
 
 		/**
-		* Selects actors and sub-actors for the given LODActor
+		* Selects actors and sub-actors and the given LODActor
 		*
 		* @param LODActor - Actor to select + subactors
 		* @param SelectionDepth - (recursive)
 		*/
 		void SelectLODActorAndContainedActorsInViewport(ALODActor* LODActor, const uint32 SelectionDepth = 0);
 
+		/**
+		* Selects actors and sub-actors for the given LODActor
+		*
+		* @param LODActor - Actor to select + subactors
+		* @param SelectionDepth - (recursive)
+		*/
+		void SelectContainedActorsInViewport(ALODActor* LODActor, const uint32 SelectionDepth = 0);
+
 		/**
 		* Creates a ASelectionActor for the given actor, "procedurally" drawing its bounds
 		*
@@ -304,16 +305,8 @@ namespace HLODOutliner
 		*/
 		UDrawSphereComponent* CreateBoundingSphereForActor(AActor* Actor);
 
-		/**
-		* Extracts all the Static Mesh Actors from the given LODActor's SubActors array
-		*
-		* @param LODActor - LODActors to check the SubActors array for
-		* @param InOutActors - Array to fill with Static Mesh Actors
-		*/
-		void ExtractStaticMeshActorsFromLODActor(ALODActor* LODActor, TArray<AActor*> &InOutActors);
-	
-		/** Ends the Editor selection batch */
-		void EndSelection();
+		/** Ends the Editor selection batch, bChange determines whether or not there was an actual change and call NoteSelectionChange */
+		void EndSelection(const bool bChange);
 
 	protected:
 		/** Broadcast event delegates */
@@ -342,6 +335,9 @@ namespace HLODOutliner
 		/** Called when the current level has changed */
 		void OnNewCurrentLevel();
 
+		/** Called when a new map is being loaded */
+		void OnMapLoaded(const FString&  Filename, bool bAsTemplate);
+
 		/** Called when a HLODActor is moved between clusters */
 		void OnHLODActorMovedEvent(const AActor* InActor, const AActor* ParentActor);
 
@@ -360,10 +356,13 @@ namespace HLODOutliner
 		/** Called when the HLOD Levels array within WorldSettings changed */
 		void OnHLODLevelsArrayChangedEvent();
 
+		/** Called when an Actor is removed from a cluster */
+		void OnHLODActorRemovedFromClusterEvent(const AActor* InActor, const AActor* ParentActor);
+
 		/** End of Broadcast event delegates */
 
 		/** Callback function used to check if Hierarchical LOD functionality is enabled in the current world settings */
-		bool IsHLODEnabledInWorldSettings();
+		bool IsHLODEnabledInWorldSettings() const;
 	private:
 		/** Tells the scene outliner that it should do a full refresh, which will clear the entire tree and rebuild it from scratch. */
 		void FullRefresh();
@@ -440,7 +439,7 @@ namespace HLODOutliner
 		TAttribute<bool> EnabledAttribute;
 	
 		/** Map containing all the nodes with their corresponding keys */
-		TMap<FTreeItemID, FTreeItemPtr> TreeItemsMap;
+		TMultiMap<FTreeItemID, FTreeItemPtr> TreeItemsMap;
 
 		/** Array of pending OutlinerActions */
 		TArray<FOutlinerAction> PendingActions;
diff --git a/Engine/Source/Editor/HierarchicalLODOutliner/Private/HLODOutlinerDragDrop.h b/Engine/Source/Editor/HierarchicalLODOutliner/Private/HLODOutlinerDragDrop.h
--- a/Engine/Source/Editor/HierarchicalLODOutliner/Private/HLODOutlinerDragDrop.h
+++ b/Engine/Source/Editor/HierarchicalLODOutliner/Private/HLODOutlinerDragDrop.h
@@ -52,22 +52,12 @@ namespace HLODOutliner
 	/** A drag/drop operation that was started from the scene outliner */
 	struct FHLODOutlinerDragDropOp : public FDragDropOperation
 	{
-		enum ToolTipTextType
+		enum ToolTipIconType
 		{
 			ToolTip_Compatible,
 			ToolTip_Incompatible,
 			ToolTip_MultipleSelection_Compatible,
-			ToolTip_MultipleSelection_Incompatible,
-			ToolTip_CompatibleMoveToCluster,
-			ToolTip_MultipleSelection_CompatibleMoveToCluster,
-			ToolTip_CompatibleAddToCluster,
-			ToolTip_MultipleSelection_CompatibleAddToCluster,
-			ToolTip_CompatibleMergeCluster,
-			ToolTip_MultipleSelection_CompatibleMergeClusters,
-			ToolTip_CompatibleChildCluster,
-			ToolTip_MultipleSelection_CompatibleChildClusters,
-			ToolTip_CompatibleNewCluster,
-			ToolTip_MultipleSelection_CompatibleNewCluster
+			ToolTip_MultipleSelection_Incompatible
 		};
 
 		DRAG_DROP_OPERATOR_TYPE(FHLODOutlinerDragDropOp, FDragDropOperation);
@@ -110,41 +100,34 @@ namespace HLODOutliner
 	struct FDragValidationInfo
 	{
 		/** The tooltip type to display on the operation */
-		FHLODOutlinerDragDropOp::ToolTipTextType TooltipType;
+		FHLODOutlinerDragDropOp::ToolTipIconType TooltipIconType;
 
 		/** The tooltip text to display on the operation */
 		FText ValidationText;
 
+		EHierarchicalLODActionType ActionType;
 
 		/** Construct this validation information out of a tooltip type and some text */
-		FDragValidationInfo(const FHLODOutlinerDragDropOp::ToolTipTextType InTooltipType, const FText InValidationText)
-			: TooltipType(InTooltipType)
+		FDragValidationInfo(const EHierarchicalLODActionType InActionType, const FHLODOutlinerDragDropOp::ToolTipIconType InTooltipIconType, const FText InValidationText)
+			: TooltipIconType(InTooltipIconType)
 			, ValidationText(InValidationText)
-		{}
+			, ActionType(InActionType)
+		{			
+		}
 
 		/** Return a generic invalid result */
 		static FDragValidationInfo Invalid()
 		{
-			return FDragValidationInfo(FHLODOutlinerDragDropOp::ToolTip_Incompatible, FText());
+			return FDragValidationInfo(EHierarchicalLODActionType::InvalidAction, FHLODOutlinerDragDropOp::ToolTip_Incompatible, FText());
 		}
 
 		/** @return true if this operation is valid, false otherwise */
 		bool IsValid() const
 		{
-			switch (TooltipType)
+			switch (TooltipIconType)
 			{
 			case FHLODOutlinerDragDropOp::ToolTip_Compatible:
 			case FHLODOutlinerDragDropOp::ToolTip_MultipleSelection_Compatible:
-			case FHLODOutlinerDragDropOp::ToolTip_CompatibleMoveToCluster:
-			case FHLODOutlinerDragDropOp::ToolTip_MultipleSelection_CompatibleMoveToCluster:
-			case FHLODOutlinerDragDropOp::ToolTip_CompatibleAddToCluster:
-			case FHLODOutlinerDragDropOp::ToolTip_MultipleSelection_CompatibleAddToCluster:
-			case FHLODOutlinerDragDropOp::ToolTip_CompatibleMergeCluster:
-			case FHLODOutlinerDragDropOp::ToolTip_MultipleSelection_CompatibleMergeClusters:
-			case FHLODOutlinerDragDropOp::ToolTip_CompatibleChildCluster:
-			case FHLODOutlinerDragDropOp::ToolTip_MultipleSelection_CompatibleChildClusters:
-			case FHLODOutlinerDragDropOp::ToolTip_CompatibleNewCluster:
-			case FHLODOutlinerDragDropOp::ToolTip_MultipleSelection_CompatibleNewCluster:
 				return true;
 			default:
 				return false;
diff --git a/Engine/Source/Editor/HierarchicalLODOutliner/Private/HLODTreeWidgetItem.cpp b/Engine/Source/Editor/HierarchicalLODOutliner/Private/HLODTreeWidgetItem.cpp
--- a/Engine/Source/Editor/HierarchicalLODOutliner/Private/HLODTreeWidgetItem.cpp
+++ b/Engine/Source/Editor/HierarchicalLODOutliner/Private/HLODTreeWidgetItem.cpp
@@ -69,7 +69,7 @@ namespace HLODOutliner
 
 		return FReply::Unhandled();
 	}
-
+	
 	FReply HandleDrop(TWeakPtr<STableViewBase> Widget, const FDragDropEvent& DragDropEvent, IDropTarget& DropTarget, FDragValidationInfo& ValidationInfo, SHLODWidgetItem* DroppedWidget, bool bApplyDrop)
 	{
 		FDragDropPayload DraggedObjects;
@@ -77,7 +77,7 @@ namespace HLODOutliner
 		if (!DraggedObjects.ParseDrag(*DragDropEvent.GetOperation()))
 		{
 			// Invalid selection
-			ValidationInfo = FDragValidationInfo(FHLODOutlinerDragDropOp::ToolTip_Incompatible, LOCTEXT("InvalidSelection", "Selection contains already clustered Actor(s)"));
+			ValidationInfo = FDragValidationInfo(EHierarchicalLODActionType::InvalidAction, FHLODOutlinerDragDropOp::ToolTip_Incompatible, LOCTEXT("InvalidSelection", "Selection contains already clustered Actor(s)"));
 			return FReply::Unhandled();
 		}
 
@@ -139,7 +139,21 @@ BEGIN_SLATE_FUNCTION_BUILD_OPTIMIZATION
 					.ColorAndOpacity(this, &SHLODWidgetItem::GetTint)
 				];
 		}		
-		else if (ColumnName == TEXT("TriangleCount") && TreeItem->GetTreeItemType() == ITreeItem::HierarchicalLODActor)
+		else if (ColumnName == TEXT("RawTriangleCount") && TreeItem->GetTreeItemType() == ITreeItem::HierarchicalLODActor)
+		{
+			FLODActorItem* Item = static_cast<FLODActorItem*>(TreeItem);
+
+			return SNew(SHorizontalBox)
+				+ SHorizontalBox::Slot()
+				.AutoWidth()
+				.VAlign(VAlign_Center)
+				[
+					SNew(STextBlock)
+					.Text(Item, &FLODActorItem::GetRawNumTrianglesAsText)
+					.ColorAndOpacity(this, &SHLODWidgetItem::GetTint)
+				];
+		}
+		else if (ColumnName == TEXT("ReducedTriangleCount") && TreeItem->GetTreeItemType() == ITreeItem::HierarchicalLODActor)
 		{
 			FLODActorItem* Item = static_cast<FLODActorItem*>(TreeItem);
 
@@ -149,10 +163,40 @@ BEGIN_SLATE_FUNCTION_BUILD_OPTIMIZATION
 				.VAlign(VAlign_Center)
 				[
 					SNew(STextBlock)
-					.Text(Item, &FLODActorItem::GetNumTrianglesAsText)
+					.Text(Item, &FLODActorItem::GetReducedNumTrianglesAsText)
 					.ColorAndOpacity(this, &SHLODWidgetItem::GetTint)
 				];
 		}
+		else if (ColumnName == TEXT("ReductionPercentage") && TreeItem->GetTreeItemType() == ITreeItem::HierarchicalLODActor)
+		{
+			FLODActorItem* Item = static_cast<FLODActorItem*>(TreeItem);
+
+			return SNew(SHorizontalBox)
+				+ SHorizontalBox::Slot()
+				.AutoWidth()
+				.VAlign(VAlign_Center)
+				[
+					SNew(STextBlock)
+					.Text(Item, &FLODActorItem::GetReductionPercentageAsText)
+					.ColorAndOpacity(this, &SHLODWidgetItem::GetTint)
+				];
+		}
+		else if (ColumnName == TEXT("Level") && TreeItem->GetTreeItemType() == ITreeItem::HierarchicalLODActor)
+		{
+			FLODActorItem* Item = static_cast<FLODActorItem*>(TreeItem);
+
+			return SNew(SHorizontalBox)
+				+ SHorizontalBox::Slot()
+				.AutoWidth()
+				.VAlign(VAlign_Center)
+				[
+					SNew(STextBlock)
+					.Text(Item, &FLODActorItem::GetLevelAsText)
+					.ColorAndOpacity(this, &SHLODWidgetItem::GetTint)
+				];
+		}
+
+
 		else
 		{
 			return SNullWidget::NullWidget;
diff --git a/Engine/Source/Editor/HierarchicalLODOutliner/Private/HierarchicalLODOutlinerModule.cpp b/Engine/Source/Editor/HierarchicalLODOutliner/Private/HierarchicalLODOutlinerModule.cpp
--- a/Engine/Source/Editor/HierarchicalLODOutliner/Private/HierarchicalLODOutlinerModule.cpp
+++ b/Engine/Source/Editor/HierarchicalLODOutliner/Private/HierarchicalLODOutlinerModule.cpp
@@ -69,8 +69,8 @@ void FHierarchicalLODOutlinerModule::OnHLODLevelsArrayChangedEvent()
 
 TSharedRef<SWidget> FHierarchicalLODOutlinerModule::CreateHLODOutlinerWidget()
 {
-	TSharedRef<HLODOutliner::SHLODOutliner> HLODWindow = SNew(HLODOutliner::SHLODOutliner);
-	return HLODWindow;
+	SAssignNew(HLODWindow, HLODOutliner::SHLODOutliner);		
+	return HLODWindow->AsShared();
 }
 
 IMPLEMENT_MODULE(FHierarchicalLODOutlinerModule, HLODPluginModule);
\ No newline at end of file
diff --git a/Engine/Source/Editor/HierarchicalLODOutliner/Private/HierarchicalLODOutlinerPrivatePCH.h b/Engine/Source/Editor/HierarchicalLODOutliner/Private/HierarchicalLODOutlinerPrivatePCH.h
--- a/Engine/Source/Editor/HierarchicalLODOutliner/Private/HierarchicalLODOutlinerPrivatePCH.h
+++ b/Engine/Source/Editor/HierarchicalLODOutliner/Private/HierarchicalLODOutlinerPrivatePCH.h
@@ -4,3 +4,4 @@
 
 #include "Core.h"
 #include "UnrealEd.h"
+#include "HierarchicalLODType.h"
diff --git a/Engine/Source/Editor/HierarchicalLODOutliner/Private/LODActorItem.cpp b/Engine/Source/Editor/HierarchicalLODOutliner/Private/LODActorItem.cpp
--- a/Engine/Source/Editor/HierarchicalLODOutliner/Private/LODActorItem.cpp
+++ b/Engine/Source/Editor/HierarchicalLODOutliner/Private/LODActorItem.cpp
@@ -12,12 +12,13 @@
 
 #define LOCTEXT_NAMESPACE "LODActorItem"
 
-HLODOutliner::FLODActorItem::FLODActorItem(ALODActor* InLODActor)
+HLODOutliner::FLODActorItem::FLODActorItem(const ALODActor* InLODActor)
 	: LODActor(InLODActor), ID(InLODActor)
 {
 	Type = ITreeItem::HierarchicalLODActor;
 }
 
+
 bool HLODOutliner::FLODActorItem::CanInteract() const
 {
 	return true;
@@ -33,13 +34,12 @@ void HLODOutliner::FLODActorItem::GenerateContextMenu(FMenuBuilder& MenuBuilder,
 
 	if (LODActor->IsDirty())
 	{
-		MenuBuilder.AddMenuEntry(LOCTEXT("BuildLODActorMesh", "Build LOD Mesh"), FText(), FSlateIcon(), FUIAction(FExecuteAction::CreateRaw(&Outliner, &SHLODOutliner::BuildLODActor)));
+		MenuBuilder.AddMenuEntry(LOCTEXT("BuildLODActorMesh", "Build Proxy Mesh"), FText(), FSlateIcon(), FUIAction(FExecuteAction::CreateRaw(&Outliner, &SHLODOutliner::BuildLODActor)));
 	}
 	else
 	{		
 		MenuBuilder.AddMenuEntry(LOCTEXT("ForceView", "ForceView"), FText(), FSlateIcon(), FUIAction(FExecuteAction::CreateRaw(&Outliner, &SHLODOutliner::ForceViewLODActor)));
-
-		MenuBuilder.AddMenuEntry(LOCTEXT("RebuildLODActorMesh", "Rebuild LOD Mesh"), FText(), FSlateIcon(), FUIAction(FExecuteAction::CreateRaw(&Outliner, &SHLODOutliner::RebuildLODActor)));
+		MenuBuilder.AddMenuEntry(LOCTEXT("BuildLODActorMesh", "Rebuild Proxy Mesh"), FText(), FSlateIcon(), FUIAction(FExecuteAction::CreateRaw(&Outliner, &SHLODOutliner::RebuildLODActor)));
 	}
 
 	MenuBuilder.AddMenuEntry(LOCTEXT("CreateHLODVolume", "Create Containing Hierarchical Volume"), FText(), FSlateIcon(), FUIAction(FExecuteAction::CreateRaw(&Outliner, &SHLODOutliner::CreateHierarchicalVolumeForActor)));
@@ -90,14 +90,53 @@ FSlateColor HLODOutliner::FLODActorItem::GetTint() const
 	return FLinearColor(1.0f, 1.0f, 1.0f);
 }
 
-FText HLODOutliner::FLODActorItem::GetNumTrianglesAsText() const
+FText HLODOutliner::FLODActorItem::GetRawNumTrianglesAsText() const
+{
+	if (LODActor.IsValid())
+	{
+		const uint32 SubActorsTriangleCount = LODActor->GetNumTrianglesInSubActors();
+		return FText::FromString(FString::FromInt(SubActorsTriangleCount));
+	}
+	else
+	{
+		return FText::FromString(TEXT("Not available"));
+	}
+}
+
+FText HLODOutliner::FLODActorItem::GetReducedNumTrianglesAsText() const
+{
+	if (LODActor.IsValid())
+	{
+		const uint32 MergedCount = LODActor->GetNumTrianglesInMergedMesh();
+		return FText::FromString(FString::FromInt(MergedCount));
+	}
+	else
+	{
+		return FText::FromString(TEXT("Not available"));
+	}
+}
+
+FText HLODOutliner::FLODActorItem::GetReductionPercentageAsText() const
 {
 	if (LODActor.IsValid())
 	{
 		const uint32 SubActorCount = LODActor->GetNumTrianglesInSubActors();
 		const uint32 MergedCount = LODActor->GetNumTrianglesInMergedMesh();
 		const float PercentageOfOriginal = ((float)MergedCount / (float)SubActorCount) * 100.0f;
-		return FText::FromString(FString::FromInt(SubActorCount) + TEXT(" (Original) - ") + FString::FromInt(MergedCount) + TEXT(" (Merged) - ") + FString::SanitizeFloat(PercentageOfOriginal) + "% (of Original)");
+		return  FText::FromString(((SubActorCount != 0) ? FString::SanitizeFloat(PercentageOfOriginal) : TEXT("0")) + "%");
+	}
+	else
+	{
+		return FText::FromString(TEXT("Not available"));
+	}
+	
+}
+
+FText HLODOutliner::FLODActorItem::GetLevelAsText() const
+{
+	if (LODActor.IsValid())
+	{
+		return FText::FromString(LODActor->SubActors.Num() ? LODActor->SubActors[0]->GetLevel()->GetOuter()->GetName() : TEXT(""));
 	}
 	else
 	{
@@ -119,12 +158,12 @@ void HLODOutliner::FLODActorItem::PopulateDragDropPayload(FDragDropPayload& Payl
 }
 
 HLODOutliner::FDragValidationInfo HLODOutliner::FLODActorItem::ValidateDrop(FDragDropPayload& DraggedObjects) const
-{
+{	
 	if (Parent.IsValid())
 	{
 		if (Parent.Pin()->GetTreeItemType() == ITreeItem::HierarchicalLODActor)
 		{
-			return FDragValidationInfo(FHLODOutlinerDragDropOp::ToolTip_Incompatible, LOCTEXT("CannotAddToChildCluster", "Cannot add to child cluster"));
+			return FDragValidationInfo(EHierarchicalLODActionType::InvalidAction, FHLODOutlinerDragDropOp::ToolTip_Incompatible, LOCTEXT("CannotAddToChildCluster", "Cannot Add to Child cluster"));
 		}
 	}
 
@@ -165,35 +204,47 @@ HLODOutliner::FDragValidationInfo HLODOutliner::FLODActorDropTarget::ValidateDro
 				{
 					if (DraggedObjects.StaticMeshActors->Num() > 1)
 					{
-						return FDragValidationInfo(FHLODOutlinerDragDropOp::ToolTip_MultipleSelection_CompatibleMoveToCluster, LOCTEXT("MoveMultipleToCluster", "Move Actors to Cluster"));
+						return FDragValidationInfo(EHierarchicalLODActionType::MoveActorToCluster, FHLODOutlinerDragDropOp::ToolTip_MultipleSelection_Compatible, LOCTEXT("MoveMultipleToCluster", "Move Actors to Cluster"));
 					}
 					else
 					{
-						return FDragValidationInfo(FHLODOutlinerDragDropOp::ToolTip_CompatibleMoveToCluster, LOCTEXT("MoveToCluster", "Move Actor to Cluster"));
+						return FDragValidationInfo(EHierarchicalLODActionType::MoveActorToCluster, FHLODOutlinerDragDropOp::ToolTip_Compatible, LOCTEXT("MoveToCluster", "Move Actor to Cluster"));
 					}
 				}
 				else
 				{
-					return FDragValidationInfo(FHLODOutlinerDragDropOp::ToolTip_Incompatible, LOCTEXT("AlreadyInCluster", "Cannot add to existing cluster"));
+					return FDragValidationInfo(EHierarchicalLODActionType::InvalidAction, FHLODOutlinerDragDropOp::ToolTip_Incompatible, LOCTEXT("AlreadyInCluster", "Cannot Add to Existing cluster"));
 				}
 			}
 			else
 			{
-				if (DraggedObjects.StaticMeshActors->Num() > 1)
+				TArray<AActor*> Actors;
+				for (auto StaticMeshActor : DraggedObjects.StaticMeshActors.GetValue())
 				{
-					return FDragValidationInfo(FHLODOutlinerDragDropOp::ToolTip_MultipleSelection_CompatibleAddToCluster, LOCTEXT("AddMultipleToCluster", "Add Actors to Cluster"));
+					Actors.Add(StaticMeshActor.Get());
 				}
-				else
+
+				Actors.Add(LODActor.Get());
+				
+				const bool MultipleActors = DraggedObjects.StaticMeshActors->Num() > 1;
+				if (FHierarchicalLODUtilities::AreActorsInSamePersistingLevel(Actors))
 				{
-					return FDragValidationInfo(FHLODOutlinerDragDropOp::ToolTip_CompatibleAddToCluster, LOCTEXT("AddToCluster", "Add Actor to Cluster"));
+					return FDragValidationInfo(EHierarchicalLODActionType::AddActorToCluster, 
+						MultipleActors ? FHLODOutlinerDragDropOp::ToolTip_MultipleSelection_Compatible : FHLODOutlinerDragDropOp::ToolTip_Compatible,
+						MultipleActors ? LOCTEXT("AddMultipleToCluster", "Add Actors to Cluster") : LOCTEXT("AddToCluster", "Add Actor to Cluster"));
 				}
-
+				else
+				{
+					return FDragValidationInfo(EHierarchicalLODActionType::InvalidAction,
+						MultipleActors ? FHLODOutlinerDragDropOp::ToolTip_MultipleSelection_Incompatible : FHLODOutlinerDragDropOp::ToolTip_Incompatible,
+						LOCTEXT("NotInSameLODLevel", "Actors are not all in the same persisting level"));
+				}				
 			}			
 		}
 
 		if (DraggedObjects.bSceneOutliner)
 		{
-			return FDragValidationInfo(FHLODOutlinerDragDropOp::ToolTip_Incompatible, LOCTEXT("AlreadyInHLOD", "Actor is already in one of the Hierarchical LOD clusters"));
+			return FDragValidationInfo(EHierarchicalLODActionType::InvalidAction, FHLODOutlinerDragDropOp::ToolTip_Incompatible, LOCTEXT("AlreadyInHLOD", "Actor is already in one of the Hierarchical LOD clusters"));
 		}
 	}
 	else if (DraggedObjects.LODActors.IsSet() && DraggedObjects.LODActors->Num() > 0)
@@ -204,7 +255,8 @@ HLODOutliner::FDragValidationInfo HLODOutliner::FLODActorDropTarget::ValidateDro
 			bool bValidForMerge = true;
 			bool bValidForChilding = true;
 			int32 FirstLODLevel = -1;
-			ULevel* Level = nullptr;
+			UObject* LevelOuter = nullptr;
+			UObject* SubActorOuter = (LODActor->SubActors.Num()) ? LODActor->SubActors[0]->GetOuter() : nullptr;
 
 			for (auto Actor : DraggedObjects.LODActors.GetValue())
 			{
@@ -237,34 +289,41 @@ HLODOutliner::FDragValidationInfo HLODOutliner::FLODActorDropTarget::ValidateDro
 					}		
 
 					// Check if in same level asset
-					if (Level == nullptr)
+					if (LevelOuter == nullptr)
 					{
-						Level = InLODActor->GetLevel();
+						LevelOuter = InLODActor->GetOuter();
 					}
-					else if (Level != InLODActor->GetLevel())
+					else if (LevelOuter != InLODActor->GetOuter())
 					{
 						bValidForMerge = false;
 						bValidForChilding = false;
 					}
+					
+					if (InLODActor->SubActors.Num() && SubActorOuter != InLODActor->SubActors[0]->GetOuter())
+					{
+						bValidForChilding = false;
+						bValidForMerge = false;
+					}					
+					
 				}
 			}
 
 			if (bValidForMerge)
 			{
-				return FDragValidationInfo((DraggedObjects.LODActors->Num() == 1) ? FHLODOutlinerDragDropOp::ToolTip_CompatibleMergeCluster : FHLODOutlinerDragDropOp::ToolTip_MultipleSelection_CompatibleMergeClusters, LOCTEXT("MergeHLODClusters", "Merge cluster(s)"));				
+				return FDragValidationInfo(EHierarchicalLODActionType::MergeClusters, (DraggedObjects.LODActors->Num() == 1) ? FHLODOutlinerDragDropOp::ToolTip_Compatible : FHLODOutlinerDragDropOp::ToolTip_MultipleSelection_Compatible, LOCTEXT("MergeHLODClusters", "Merge Cluster(s)"));
 			}
 			else if (bValidForChilding)
 			{
-				return FDragValidationInfo((DraggedObjects.LODActors->Num() == 1) ? FHLODOutlinerDragDropOp::ToolTip_CompatibleChildCluster : FHLODOutlinerDragDropOp::ToolTip_MultipleSelection_CompatibleChildClusters, LOCTEXT("ChildHLODClusters", "Child cluster(s)"));
+				return FDragValidationInfo(EHierarchicalLODActionType::ChildCluster, (DraggedObjects.LODActors->Num() == 1) ? FHLODOutlinerDragDropOp::ToolTip_Compatible : FHLODOutlinerDragDropOp::ToolTip_MultipleSelection_Compatible, LOCTEXT("ChildHLODClusters", "Add Child Cluster(s)"));
 			}
 			else
 			{
-				return FDragValidationInfo(FHLODOutlinerDragDropOp::ToolTip_Incompatible, LOCTEXT("InvalidOperation", "Invalid Operation"));//"Cannot merge clusters from different Hierarchical LOD levels"));
+				return FDragValidationInfo(EHierarchicalLODActionType::InvalidAction, FHLODOutlinerDragDropOp::ToolTip_Incompatible, LOCTEXT("InvalidOperation", "Invalid Operation"));
 			}
 		}
 	}
 
-	return FDragValidationInfo(FHLODOutlinerDragDropOp::ToolTip_Incompatible, LOCTEXT("NotImplemented", "Not implemented"));
+	return FDragValidationInfo(EHierarchicalLODActionType::InvalidAction, FHLODOutlinerDragDropOp::ToolTip_Incompatible, LOCTEXT("NotImplemented", "Not implemented"));
 }
 
 void HLODOutliner::FLODActorDropTarget::OnDrop(FDragDropPayload& DraggedObjects, const FDragValidationInfo& ValidationInfo, TSharedRef<SWidget> DroppedOnWidget)
@@ -277,78 +336,30 @@ void HLODOutliner::FLODActorDropTarget::OnDrop(FDragDropPayload& DraggedObjects,
 
 	auto& DraggedStaticMeshActors = DraggedObjects.StaticMeshActors.GetValue();
 	auto& DraggedLODActors = DraggedObjects.LODActors.GetValue();
-	if (ValidationInfo.TooltipType == FHLODOutlinerDragDropOp::ToolTip_CompatibleMoveToCluster || ValidationInfo.TooltipType == FHLODOutlinerDragDropOp::ToolTip_MultipleSelection_CompatibleMoveToCluster)
-	{
-		for (int32 ActorIndex = 0; ActorIndex < DraggedStaticMeshActors.Num(); ++ActorIndex)
-		{
-			auto Actor = DraggedStaticMeshActors[ActorIndex];
-			MoveToCluster(Actor.Get(), nullptr, LODActor.Get());
-		}
-	}
-	else if (ValidationInfo.TooltipType == FHLODOutlinerDragDropOp::ToolTip_CompatibleAddToCluster || ValidationInfo.TooltipType == FHLODOutlinerDragDropOp::ToolTip_MultipleSelection_CompatibleAddToCluster)
+	if (ValidationInfo.ActionType == EHierarchicalLODActionType::AddActorToCluster || ValidationInfo.ActionType == EHierarchicalLODActionType::MoveActorToCluster)
 	{
 		for (int32 ActorIndex = 0; ActorIndex < DraggedStaticMeshActors.Num(); ++ActorIndex)
 		{
 			auto Actor = DraggedStaticMeshActors[ActorIndex];
-			AddToCluster(Actor.Get(), LODActor.Get());
+			FHierarchicalLODUtilities::AddActorToCluster(Actor.Get(), LODActor.Get());
 		}
 	}
-	else if (ValidationInfo.TooltipType == FHLODOutlinerDragDropOp::ToolTip_CompatibleMergeCluster || ValidationInfo.TooltipType == FHLODOutlinerDragDropOp::ToolTip_MultipleSelection_CompatibleMergeClusters)
+	else if (ValidationInfo.ActionType == EHierarchicalLODActionType::MergeClusters)
 	{
 		for (int32 ActorIndex = 0; ActorIndex < DraggedLODActors.Num(); ++ActorIndex)
 		{
-			ALODActor* InLODActor = Cast<ALODActor>(DraggedLODActors[ActorIndex].Get());
-			MergeCluster(InLODActor);
+			ALODActor* InLODActor = Cast<ALODActor>(DraggedLODActors[ActorIndex].Get());			
+			FHierarchicalLODUtilities::MergeClusters(LODActor.Get(), InLODActor);
 		}
 	}
-	else if (ValidationInfo.TooltipType == FHLODOutlinerDragDropOp::ToolTip_CompatibleChildCluster || ValidationInfo.TooltipType == FHLODOutlinerDragDropOp::ToolTip_MultipleSelection_CompatibleChildClusters)
+	else if (ValidationInfo.ActionType == EHierarchicalLODActionType::ChildCluster)
 	{
 		for (int32 ActorIndex = 0; ActorIndex < DraggedLODActors.Num(); ++ActorIndex)
 		{
-			auto Actor = DraggedLODActors[ActorIndex];
-			MoveToCluster(Actor.Get(), nullptr, LODActor.Get());
+			auto Actor = DraggedLODActors[ActorIndex];			
+			FHierarchicalLODUtilities::AddActorToCluster(Actor.Get(), LODActor.Get());
 		}
 	}
 }
 
-void HLODOutliner::FLODActorDropTarget::MoveToCluster(AActor* InActor, ALODActor* OldParentActor, ALODActor* NewParentActor)
-{
-	const FScopedTransaction Transaction(LOCTEXT("UndoAction_MoveActorBetweenClusters", "Move Actor between Clusters"));
-	InActor->Modify();
-	ALODActor* CurrentParentActor = (OldParentActor) ? OldParentActor : FHierarchicalLODUtilities::GetParentLODActor(InActor);
-	if (CurrentParentActor)
-	{
-		CurrentParentActor->Modify();
-		CurrentParentActor->RemoveSubActor(InActor);
-	}
-	
-	NewParentActor->Modify();
-	NewParentActor->AddSubActor(InActor);
-	
-	GEngine->BroadcastHLODActorMoved(InActor, NewParentActor);
-}
-
-void HLODOutliner::FLODActorDropTarget::AddToCluster(AActor* InActor, ALODActor* NewParentActor)
-{
-	const FScopedTransaction Transaction(LOCTEXT("UndoAction_AddActorToCluster", "Add Actor To Cluster"));
-	NewParentActor->Modify();
-	InActor->Modify();
-
-	NewParentActor->AddSubActor(InActor);
-	GEngine->BroadcastHLODActorAdded(InActor, NewParentActor);
-}
-
-void HLODOutliner::FLODActorDropTarget::MergeCluster(ALODActor* ToMergeActor)
-{
-	const FScopedTransaction Transaction(LOCTEXT("UndoAction_MergeClusters", "Merge Clusters"));
-	ToMergeActor->Modify();
-	LODActor->Modify();
-
-	while (ToMergeActor->SubActors.Num())
-	{
-		AActor* SubActor = ToMergeActor->SubActors.Last();
-		MoveToCluster(SubActor, ToMergeActor, LODActor.Get());
-	}
-}
-
 #undef LOCTEXT_NAMESPACE
diff --git a/Engine/Source/Editor/HierarchicalLODOutliner/Private/LODActorItem.h b/Engine/Source/Editor/HierarchicalLODOutliner/Private/LODActorItem.h
--- a/Engine/Source/Editor/HierarchicalLODOutliner/Private/LODActorItem.h
+++ b/Engine/Source/Editor/HierarchicalLODOutliner/Private/LODActorItem.h
@@ -24,15 +24,6 @@ namespace HLODOutliner
 
 		/** Called to drop the specified objects on this item. Only called if ValidateDrop() allows. */
 		virtual void OnDrop(FDragDropPayload& DraggedObjects, const FDragValidationInfo& ValidationInfo, TSharedRef<SWidget> DroppedOnWidget) override;
-	private:
-		/** Move a static mesh actor to a different LODActor (cluster) */
-		void MoveToCluster(AActor* InActor, ALODActor* OldParentActor, ALODActor* NewParentActor);
-
-		/** Add a static mesh actor to a LODActor (cluster) */
-		void AddToCluster(AActor* InActor, ALODActor* NewParentActor);
-
-		/** Merges ToMergeActor in the represented Cluster/LODActor */
-		void MergeCluster(ALODActor* ToMergeActor);
 	};
 
 	struct FLODActorItem : ITreeItem
@@ -40,7 +31,7 @@ namespace HLODOutliner
 		mutable TWeakObjectPtr<ALODActor> LODActor;
 		mutable FTreeItemID ID;
 
-		FLODActorItem(ALODActor* InLODActor);
+		FLODActorItem(const ALODActor* InLODActor);
 
 		//~ Begin ITreeItem Interface.
 		virtual bool CanInteract() const override;
@@ -51,7 +42,18 @@ namespace HLODOutliner
 		//~ End ITreeItem Interface.
 
 		/** Returns the number of triangles for the represented LODActor in FText form */
-		FText GetNumTrianglesAsText() const;
+		FText GetRawNumTrianglesAsText() const;
+
+		/** Returns the reduced number of triangles for the represented LODActor in FText form */
+		FText GetReducedNumTrianglesAsText() const;
+
+		/** Returns the reduction percentage for the represented LODActor in FText form */
+		FText GetReductionPercentageAsText() const;
+
+		/** Returns the level the meshes are in Ftext form */
+		FText GetLevelAsText() const;
+
+		
 
 		/** Populate the specified drag/drop payload with any relevant information for this type */
 		virtual void PopulateDragDropPayload(FDragDropPayload& Payload) const override;
