Copying //UE4/Dev-Framework to Dev-Main (//UE4/Dev-Main) @ 2855699

diff --git a/Engine/Source/Editor/HierarchicalLODOutliner/Private/LODLevelItem.cpp b/Engine/Source/Editor/HierarchicalLODOutliner/Private/LODLevelItem.cpp
--- a/Engine/Source/Editor/HierarchicalLODOutliner/Private/LODLevelItem.cpp
+++ b/Engine/Source/Editor/HierarchicalLODOutliner/Private/LODLevelItem.cpp
@@ -25,7 +25,7 @@ bool HLODOutliner::FLODLevelItem::CanInteract() const
 
 void HLODOutliner::FLODLevelItem::GenerateContextMenu(FMenuBuilder& MenuBuilder, SHLODOutliner& Outliner)
 {
-
+	// No context menu available for LODLevel item
 }
 
 FString HLODOutliner::FLODLevelItem::GetDisplayString() const
@@ -63,37 +63,30 @@ HLODOutliner::FDragValidationInfo HLODOutliner::FLODLevelDropTarget::ValidateDro
 	if (DraggedObjects.StaticMeshActors.IsSet() && DraggedObjects.StaticMeshActors->Num() > 0)
 	{
 		const int32 NumStaticMeshActors = DraggedObjects.StaticMeshActors->Num();	
-		bool bSameLevelInstance = true;
-		bool bAlreadyClustered = false;
+		
 		ULevel* Level = nullptr;
+
+		TArray<AActor*> DraggedActors;
+
 		for (auto Actor : DraggedObjects.StaticMeshActors.GetValue())
 		{
-			if (Level == nullptr)
-			{
-				Level = Actor->GetLevel();
-			}
-			else if (Level != Actor->GetLevel())
-			{
-				bSameLevelInstance = false;
-			}
-
-			if (FHierarchicalLODUtilities::GetParentLODActor(Actor.Get()))
-			{
-				bAlreadyClustered = true;
-			}
+			DraggedActors.Add(Actor.Get());
 		}
 
+		bool bSameLevelInstance = FHierarchicalLODUtilities::AreActorsInSamePersistingLevel(DraggedActors);
+		bool bAlreadyClustered = FHierarchicalLODUtilities::AreActorsClustered(DraggedActors);
+		
 		if (!bSameLevelInstance)
 		{
-			return FDragValidationInfo(FHLODOutlinerDragDropOp::ToolTip_Incompatible, LOCTEXT("StaticMeshActorsNotInSameLevelAsset", "Static Mesh Actors not in the same level asset (streaming level)"));
+			return FDragValidationInfo(EHierarchicalLODActionType::InvalidAction, FHLODOutlinerDragDropOp::ToolTip_Incompatible, LOCTEXT("StaticMeshActorsNotInSameLevelAsset", "Static Mesh Actors not in the same level asset (streaming level)"));
 		}
 
-		if (bAlreadyClustered)
+		if (bAlreadyClustered && DraggedObjects.bSceneOutliner )
 		{
-			return FDragValidationInfo(FHLODOutlinerDragDropOp::ToolTip_Incompatible, LOCTEXT("AlreadyClusters", "One or more Static Mesh Actors is already in a cluster"));
+			return FDragValidationInfo(EHierarchicalLODActionType::InvalidAction, FHLODOutlinerDragDropOp::ToolTip_Incompatible, LOCTEXT("AlreadyClusters", "One or more Static Mesh Actors is already in a cluster"));
 		}
 
-		return FDragValidationInfo(FHLODOutlinerDragDropOp::ToolTip_CompatibleNewCluster, LOCTEXT("CreateNewCluster", "Create new Cluster"));
+		return FDragValidationInfo(EHierarchicalLODActionType::CreateCluster, FHLODOutlinerDragDropOp::ToolTip_Compatible, LOCTEXT("CreateNewCluster", "Create new Cluster"));
 	}
 	else if (DraggedObjects.LODActors.IsSet() && DraggedObjects.LODActors->Num() > 0)
 	{
@@ -101,71 +94,46 @@ HLODOutliner::FDragValidationInfo HLODOutliner::FLODLevelDropTarget::ValidateDro
 
 		if (NumLODActors > 1)
 		{
-			// Check if all the dragged LOD actors fall within the same LOD level
-			auto LODActors = DraggedObjects.LODActors.GetValue();
-			int32 LevelIndex = -1;
-			bool bSameLODLevel = true;
-			bool bSameLevelInstance = true;
-			bool bValidOperation = true;
-			ULevel* Level = nullptr;
+			// Gather LOD actors for checks
+			auto LODActors = DraggedObjects.LODActors.GetValue();			
+			TArray<ALODActor*> DraggedLODActors;
+			TArray<AActor*> DraggedActors;
 			for (auto Actor : LODActors)
 			{
-				ALODActor* LODActor = Cast<ALODActor>(Actor.Get());
-				if (LevelIndex == -1)
-				{
-					LevelIndex = LODActor->LODLevel;
-				}
-				else if (LevelIndex != LODActor->LODLevel)
-				{
-					bSameLODLevel = false;
-				}
-
-				if (FHierarchicalLODUtilities::GetParentLODActor(LODActor))
-				{
-					bValidOperation = false;
-				}
-
-				if (Level == nullptr)
-				{
-					Level = LODActor->GetLevel();
-				}
-				else if (Level != LODActor->GetLevel())
-				{
-					bSameLevelInstance = false;
-				}
+				DraggedLODActors.Add(Cast<ALODActor>(Actor.Get()));
+				DraggedActors.Add(Actor.Get());
 			}
 
-			if (!bValidOperation)
-			{
-				return FDragValidationInfo(FHLODOutlinerDragDropOp::ToolTip_Incompatible, LOCTEXT("OneOrMultipleAlreadyInHLODLevel", "One or multiple LODActors are already part of a cluster in this HLOD level"));
-			}
+			bool bSameLevelInstance = FHierarchicalLODUtilities::AreActorsInSamePersistingLevel(DraggedActors);
+			bool bSameLODLevel = FHierarchicalLODUtilities::AreClustersInSameHLODLevel(DraggedLODActors);
+			const uint32 LevelIndex = (DraggedLODActors.Num() > 0 ) ? DraggedLODActors[0]->LODLevel : 0;
+			bool bIsClustered = FHierarchicalLODUtilities::AreActorsClustered(DraggedActors);
 
 			if (!bSameLODLevel)
 			{
-				return FDragValidationInfo(FHLODOutlinerDragDropOp::ToolTip_Incompatible, LOCTEXT("NotInSameLODLevel", "LODActors are not all in the same HLOD level"));
+				return FDragValidationInfo(EHierarchicalLODActionType::InvalidAction, FHLODOutlinerDragDropOp::ToolTip_Incompatible, LOCTEXT("NotInSameLODLevel", "LODActors are not all in the same HLOD level"));
 			}
 
 			if (!bSameLevelInstance)
 			{
-				return FDragValidationInfo(FHLODOutlinerDragDropOp::ToolTip_Incompatible, LOCTEXT("LODActorsNotInSameLevelAsset", "LODActors not in the same level asset (streaming level)"));
+				return FDragValidationInfo(EHierarchicalLODActionType::InvalidAction, FHLODOutlinerDragDropOp::ToolTip_Incompatible, LOCTEXT("LODActorsNotInSameLevelAsset", "LODActors not in the same level asset (streaming level)"));
 			}
-
+						
 			if (bSameLevelInstance && bSameLODLevel && LevelIndex < (int32)(LODLevelIndex + 1))
 			{
-				return FDragValidationInfo(FHLODOutlinerDragDropOp::ToolTip_MultipleSelection_CompatibleNewCluster, LOCTEXT("CreateNewCluster", "Create new Cluster"));
-			}
-			
+				return FDragValidationInfo(EHierarchicalLODActionType::CreateCluster, FHLODOutlinerDragDropOp::ToolTip_MultipleSelection_Compatible, LOCTEXT("CreateNewCluster", "Create new Cluster"));
+			}			
 		}
 	}
 
-	return FDragValidationInfo(FHLODOutlinerDragDropOp::ToolTip_Incompatible, LOCTEXT("NotImplemented", "Not implemented"));
-}
+	return FDragValidationInfo(EHierarchicalLODActionType::InvalidAction, FHLODOutlinerDragDropOp::ToolTip_Incompatible, LOCTEXT("NotImplemented", "Not implemented"));
+}				
 
 void HLODOutliner::FLODLevelDropTarget::OnDrop(FDragDropPayload& DraggedObjects, const FDragValidationInfo& ValidationInfo, TSharedRef<SWidget> DroppedOnWidget)
 {
-	if (ValidationInfo.TooltipType == FHLODOutlinerDragDropOp::ToolTip_CompatibleNewCluster || ValidationInfo.TooltipType == FHLODOutlinerDragDropOp::ToolTip_MultipleSelection_CompatibleNewCluster)
+	if (ValidationInfo.ActionType == EHierarchicalLODActionType::CreateCluster)
 	{
-		CreateNewCluster(DraggedObjects);		
+		CreateNewCluster(DraggedObjects);
 	}	
 }
 
@@ -182,61 +150,20 @@ void HLODOutliner::FLODLevelDropTarget::CreateNewCluster(FDragDropPayload &Dragg
 		OuterWorld = Cast<UWorld>(DraggedObjects.LODActors.GetValue()[0]->GetLevel()->GetOuter());
 	}
 
-	// Retrieve world settings from the InWorld instance, this is the instance the HLODOutliner is running on
-	auto WorldSettings = DraggedObjects.OutlinerWorld->GetWorldSettings();
-
-	const FScopedTransaction Transaction(LOCTEXT("UndoAction_CreateNewCluster", "Create new Cluster"));
-	OuterWorld->Modify();
-
-	if (WorldSettings->bEnableHierarchicalLODSystem)
+	// Gather sub actors from the drag and drop operation
+	TArray<AActor*> SubActors;
+	for (TWeakObjectPtr<AActor> StaticMeshActor : DraggedObjects.StaticMeshActors.GetValue())
 	{
-		ALODActor* NewCluster = FHierarchicalLODUtilities::CreateNewClusterActor(OuterWorld, LODLevelIndex, WorldSettings);
-
-		if (NewCluster)
-		{
-			for (TWeakObjectPtr<AActor> StaticMeshActor : DraggedObjects.StaticMeshActors.GetValue())
-			{
-				AActor* InActor = StaticMeshActor.Get();
-				checkf(InActor != nullptr, TEXT("Null-ed StaticMeshActor found in drag and drop operation"));
-
-				InActor->Modify();
-				ALODActor* CurrentParentActor = FHierarchicalLODUtilities::GetParentLODActor(InActor);
-				if (CurrentParentActor)
-				{
-					CurrentParentActor->RemoveSubActor(InActor);
-					CurrentParentActor->Modify();
-					if (!CurrentParentActor->HasValidSubActors())
-					{
-						FHierarchicalLODUtilities::DeleteLODActor(CurrentParentActor);
-					}
-				}
-
-				NewCluster->AddSubActor(InActor);
-			}
-
-			for (TWeakObjectPtr<AActor> LODActor : DraggedObjects.LODActors.GetValue())
-			{
-				AActor* InActor = LODActor.Get();
-				checkf(InActor != nullptr, TEXT("Null-ed StaticMeshActor found in drag and drop operation"));
-				ALODActor* CurrentParentActor = FHierarchicalLODUtilities::GetParentLODActor(InActor);
-				InActor->Modify();
-				if (CurrentParentActor)
-				{
-					CurrentParentActor->RemoveSubActor(InActor);
-					CurrentParentActor->Modify();
-					if (!CurrentParentActor->HasValidSubActors())
-					{
-						FHierarchicalLODUtilities::DeleteLODActor(CurrentParentActor);
-					}
-				}
-
-				NewCluster->AddSubActor(InActor);
-			}
+		SubActors.Add(StaticMeshActor.Get());
+	}
 
-			// Update sub actor LOD parents to populate 
-			NewCluster->UpdateSubActorLODParents();
-		}
+	for (TWeakObjectPtr<AActor> LODActor : DraggedObjects.LODActors.GetValue())
+	{
+		SubActors.Add(LODActor.Get());
 	}
+
+	// Create the new cluster
+	FHierarchicalLODUtilities::CreateNewClusterFromActors(OuterWorld, DraggedObjects.OutlinerWorld->GetWorldSettings(), SubActors, LODLevelIndex);	
 }
 
 #undef LOCTEXT_NAMESPACE
diff --git a/Engine/Source/Editor/HierarchicalLODOutliner/Private/StaticMeshActorItem.cpp b/Engine/Source/Editor/HierarchicalLODOutliner/Private/StaticMeshActorItem.cpp
--- a/Engine/Source/Editor/HierarchicalLODOutliner/Private/StaticMeshActorItem.cpp
+++ b/Engine/Source/Editor/HierarchicalLODOutliner/Private/StaticMeshActorItem.cpp
@@ -10,7 +10,7 @@
 
 #define LOCTEXT_NAMESPACE "StaticMeshActorItem"
 
-HLODOutliner::FStaticMeshActorItem::FStaticMeshActorItem(AActor* InStaticMeshActor)
+HLODOutliner::FStaticMeshActorItem::FStaticMeshActorItem(const AActor* InStaticMeshActor)
 	: StaticMeshActor(InStaticMeshActor), ID(InStaticMeshActor)
 {
 	Type = ITreeItem::StaticMeshActor;
@@ -25,7 +25,7 @@ void HLODOutliner::FStaticMeshActorItem::GenerateContextMenu(FMenuBuilder& MenuB
 {
 	auto SharedOutliner = StaticCastSharedRef<SHLODOutliner>(Outliner.AsShared());
 	MenuBuilder.AddMenuEntry(LOCTEXT("RemoveSMActorFromCluster", "Remove From Cluster"), FText(), FSlateIcon(), FUIAction(FExecuteAction::CreateRaw(&Outliner, &SHLODOutliner::RemoveStaticMeshActorFromCluster)));
-	MenuBuilder.AddMenuEntry(LOCTEXT("ExcludeSMActorFromGeneration", "Exclude from Cluster Generation"), FText(), FSlateIcon(), FUIAction(FExecuteAction::CreateRaw(&Outliner, &SHLODOutliner::ExcludeFromClusterGeneration)));
+	MenuBuilder.AddMenuEntry(LOCTEXT("ExcludeSMActorFromGeneration", "Exclude From Cluster Generation"), FText(), FSlateIcon(), FUIAction(FExecuteAction::CreateRaw(&Outliner, &SHLODOutliner::ExcludeFromClusterGeneration)));
 }
 
 FString HLODOutliner::FStaticMeshActorItem::GetDisplayString() const
@@ -73,7 +73,7 @@ void HLODOutliner::FStaticMeshActorItem::OnDrop(FDragDropPayload& DraggedObjects
 
 HLODOutliner::FDragValidationInfo HLODOutliner::FStaticMeshActorDropTarget::ValidateDrop(FDragDropPayload& DraggedObjects) const
 {
-	return FDragValidationInfo(FHLODOutlinerDragDropOp::ToolTip_Incompatible, LOCTEXT("NotImplemented", "Not implemented"));
+	return FDragValidationInfo(EHierarchicalLODActionType::InvalidAction, FHLODOutlinerDragDropOp::ToolTip_Incompatible, LOCTEXT("NotImplemented", "Not implemented"));
 }
 
 void HLODOutliner::FStaticMeshActorDropTarget::OnDrop(FDragDropPayload& DraggedObjects, const FDragValidationInfo& ValidationInfo, TSharedRef<SWidget> DroppedOnWidget)
diff --git a/Engine/Source/Editor/HierarchicalLODOutliner/Private/StaticMeshActorItem.h b/Engine/Source/Editor/HierarchicalLODOutliner/Private/StaticMeshActorItem.h
--- a/Engine/Source/Editor/HierarchicalLODOutliner/Private/StaticMeshActorItem.h
+++ b/Engine/Source/Editor/HierarchicalLODOutliner/Private/StaticMeshActorItem.h
@@ -33,7 +33,7 @@ namespace HLODOutliner
 		/** TreeItem's ID */
 		mutable FTreeItemID ID;
 
-		FStaticMeshActorItem(AActor* InStaticMeshActor);
+		FStaticMeshActorItem(const AActor* InStaticMeshActor);
 
 		//~ Begin ITreeItem Interface.
 		virtual bool CanInteract() const override;
diff --git a/Engine/Source/Editor/HierarchicalLODOutliner/Public/HierarchicalLODOutlinerModule.h b/Engine/Source/Editor/HierarchicalLODOutliner/Public/HierarchicalLODOutlinerModule.h
--- a/Engine/Source/Editor/HierarchicalLODOutliner/Public/HierarchicalLODOutlinerModule.h
+++ b/Engine/Source/Editor/HierarchicalLODOutliner/Public/HierarchicalLODOutlinerModule.h
@@ -3,6 +3,9 @@
 #include "ModuleInterface.h"
 #include "UnrealEd.h"
 
+
+namespace HLODOutliner { class SHLODOutliner; };
+
 /**
 * The module holding all of the UI related pieces for HLOD Outliner
 */
@@ -26,4 +29,5 @@ public:
 
 private:
 	FDelegateHandle ArrayChangedDelegate;
+	TSharedPtr<HLODOutliner::SHLODOutliner> HLODWindow;
 };
\ No newline at end of file
diff --git a/Engine/Source/Editor/Kismet/Private/SSCSEditor.cpp b/Engine/Source/Editor/Kismet/Private/SSCSEditor.cpp
--- a/Engine/Source/Editor/Kismet/Private/SSCSEditor.cpp
+++ b/Engine/Source/Editor/Kismet/Private/SSCSEditor.cpp
@@ -536,7 +536,7 @@ void FSCSEditorTreeNode::AddChild(FSCSEditorTreeNodePtrType InChildNodePtr)
 			if (ensure(ParentInstance != nullptr))
 			{
 				// Handle attachment at the instance level
-				if (ChildInstance->AttachParent != ParentInstance)
+				if (ChildInstance->GetAttachParent() != ParentInstance)
 				{
 					AActor* Owner = ParentInstance->GetOwner();
 					if (Owner->GetRootComponent() == ChildInstance)
@@ -2404,7 +2404,7 @@ void SSCS_RowWidget::OnAttachToDropAction(const TArray<FSCSEditorTreeNodePtrType
 				{
 					// If we find a match, calculate its new position relative to the scene root component instance in its current scene
 					FTransform ComponentToWorld(SceneComponentTemplate->RelativeRotation, SceneComponentTemplate->RelativeLocation, SceneComponentTemplate->RelativeScale3D);
-					FTransform ParentToWorld = SceneComponentTemplate->AttachSocketName != NAME_None ? ParentSceneComponent->GetSocketTransform(SceneComponentTemplate->AttachSocketName, RTS_World) : ParentSceneComponent->GetComponentToWorld();
+					FTransform ParentToWorld = SceneComponentTemplate->GetAttachSocketName() != NAME_None ? ParentSceneComponent->GetSocketTransform(SceneComponentTemplate->GetAttachSocketName(), RTS_World) : ParentSceneComponent->GetComponentToWorld();
 					FTransform RelativeTM = ComponentToWorld.GetRelativeTransform(ParentToWorld);
 
 					// Store new relative location value (if not set to absolute)
@@ -2539,7 +2539,7 @@ void SSCS_RowWidget::OnDetachFromDropAction(const TArray<FSCSEditorTreeNodePtrTy
 			{
 				// If we find a match, calculate its new position relative to the scene root component instance in the preview scene
 				FTransform ComponentToWorld(SceneComponentTemplate->RelativeRotation, SceneComponentTemplate->RelativeLocation, SceneComponentTemplate->RelativeScale3D);
-				FTransform ParentToWorld = SceneComponentTemplate->AttachSocketName != NAME_None ? InstancedSceneRootComponent->GetSocketTransform(SceneComponentTemplate->AttachSocketName, RTS_World) : InstancedSceneRootComponent->GetComponentToWorld();
+				FTransform ParentToWorld = SceneComponentTemplate->GetAttachSocketName() != NAME_None ? InstancedSceneRootComponent->GetSocketTransform(SceneComponentTemplate->GetAttachSocketName(), RTS_World) : InstancedSceneRootComponent->GetComponentToWorld();
 				FTransform RelativeTM = ComponentToWorld.GetRelativeTransform(ParentToWorld);
 
 				// Store new relative location value (if not set to absolute)
@@ -4022,9 +4022,9 @@ FSCSEditorTreeNodePtrType SSCSEditor::GetNodeFromActorComponent(const UActorComp
 		if(!NodePtr.IsValid() && bIncludeAttachedComponents)
 		{
 			const USceneComponent* SceneComponent = Cast<const USceneComponent>(ActorComponent);
-			if(SceneComponent && SceneComponent->AttachParent)
+			if(SceneComponent && SceneComponent->GetAttachParent())
 			{
-				return GetNodeFromActorComponent(SceneComponent->AttachParent, bIncludeAttachedComponents);
+				return GetNodeFromActorComponent(SceneComponent->GetAttachParent(), bIncludeAttachedComponents);
 			}
 		}
 	}
@@ -5261,21 +5261,21 @@ FSCSEditorTreeNodePtrType SSCSEditor::AddTreeNodeFromComponent(USceneComponent*
 	ensure(!InSceneComponent->IsPendingKill());
 
 	// If the given component has a parent, and if we're not in "instance" mode OR the owner of the parent matches the Actor instance we're editing
-	if(InSceneComponent->AttachParent != NULL
-		&& (EditorMode != EComponentEditorMode::ActorInstance || InSceneComponent->AttachParent->GetOwner() == GetActorContext()))
+	if(InSceneComponent->GetAttachParent() != NULL
+		&& (EditorMode != EComponentEditorMode::ActorInstance || InSceneComponent->GetAttachParent()->GetOwner() == GetActorContext()))
 	{
 		// Attempt to find the parent node in the current tree
-		FSCSEditorTreeNodePtrType ParentNodePtr = FindTreeNode(InSceneComponent->AttachParent);
+		FSCSEditorTreeNodePtrType ParentNodePtr = FindTreeNode(InSceneComponent->GetAttachParent());
 		if(!ParentNodePtr.IsValid())
 		{
 			// If the actual attach parent wasn't found, attempt to find its archetype.
 			// This handles the BP editor case where we might add UCS component nodes taken
 			// from the preview actor instance, which are not themselves template objects.
-			ParentNodePtr = FindTreeNode(Cast<USceneComponent>(InSceneComponent->AttachParent->GetArchetype()));
+			ParentNodePtr = FindTreeNode(Cast<USceneComponent>(InSceneComponent->GetAttachParent()->GetArchetype()));
 			if(!ParentNodePtr.IsValid())
 			{
 				// Recursively add the parent node to the tree if it does not exist yet
-				ParentNodePtr = AddTreeNodeFromComponent(InSceneComponent->AttachParent);
+				ParentNodePtr = AddTreeNodeFromComponent(InSceneComponent->GetAttachParent());
 			}
 		}
 
diff --git a/Engine/Source/Editor/KismetCompiler/Private/AnimBlueprintCompiler.cpp b/Engine/Source/Editor/KismetCompiler/Private/AnimBlueprintCompiler.cpp
--- a/Engine/Source/Editor/KismetCompiler/Private/AnimBlueprintCompiler.cpp
+++ b/Engine/Source/Editor/KismetCompiler/Private/AnimBlueprintCompiler.cpp
@@ -557,8 +557,9 @@ void FAnimBlueprintCompiler::ProcessAllAnimationNodes()
 	TArray<UK2Node_TransitionRuleGetter*> Getters;
 	ConsolidatedEventGraph->GetNodesOfClass<UK2Node_TransitionRuleGetter>(/*out*/ Getters);
 
-	TArray<UK2Node_AnimGetter*> AnimGetters;
-	ConsolidatedEventGraph->GetNodesOfClass<UK2Node_AnimGetter>(AnimGetters);
+	// Get anim getters from the root anim graph (processing the nodes below will collect them in nested graphs)
+	TArray<UK2Node_AnimGetter*> RootGraphAnimGetters;
+	ConsolidatedEventGraph->GetNodesOfClass<UK2Node_AnimGetter>(RootGraphAnimGetters);
 
 	// Find the root node
 	UAnimGraphNode_Root* PrePhysicsRoot = NULL;
@@ -605,9 +606,16 @@ void FAnimBlueprintCompiler::ProcessAllAnimationNodes()
 			ProcessTransitionGetter(*GetterIt, NULL); // transition nodes should not appear at top-level
 		}
 
-		for(auto AnimGetterIt = AnimGetters.CreateIterator(); AnimGetterIt; ++AnimGetterIt)
+		// Wire root getters
+		for(UK2Node_AnimGetter* RootGraphGetter : RootGraphAnimGetters)
 		{
-			AutoWireAnimGetter(*AnimGetterIt, NULL);
+			AutoWireAnimGetter(RootGraphGetter, nullptr);
+		}
+
+		// Wire nested getters
+		for(UK2Node_AnimGetter* Getter : FoundGetterNodes)
+		{
+			AutoWireAnimGetter(Getter, nullptr);
 		}
 
 		NewAnimBlueprintClass->RootAnimNodeIndex = GetAllocationIndexOfNode(PrePhysicsRoot);
@@ -677,9 +685,9 @@ int32 FAnimBlueprintCompiler::ExpandGraphAndProcessNodes(UEdGraph* SourceGraph,
 	}
 
 	// Wire anim getter nodes
-	for(auto AnimGetterIt = AnimGetterNodes.CreateIterator(); AnimGetterIt; ++AnimGetterIt)
+	for(UK2Node_AnimGetter* GetterNode : AnimGetterNodes)
 	{
-		AutoWireAnimGetter(*AnimGetterIt, TransitionNode);
+		FoundGetterNodes.Add(GetterNode);
 	}
 
 	// Returns the index of the processed cloned version of SourceRootNode
diff --git a/Engine/Source/Editor/LevelEditor/Private/SLevelViewportToolBar.cpp b/Engine/Source/Editor/LevelEditor/Private/SLevelViewportToolBar.cpp
--- a/Engine/Source/Editor/LevelEditor/Private/SLevelViewportToolBar.cpp
+++ b/Engine/Source/Editor/LevelEditor/Private/SLevelViewportToolBar.cpp
@@ -1225,10 +1225,27 @@ static void BuildBufferVisualizationMenu( FMenuBuilder& Menu )
 
 void SLevelViewportToolBar::CreateViewMenuExtensions(FMenuBuilder& MenuBuilder)
 {
+	{
+		struct Local
+		{
+			static void BuildLODMenu(FMenuBuilder& Menu, SLevelViewportToolBar* Toolbar)
+			{
+				Menu.BeginSection("LevelViewportLODColoration", LOCTEXT("LODModesHeader", "Level of Detail Coloration"));
+				{
+					Menu.AddMenuEntry(FEditorViewportCommands::Get().LODColorationMode, NAME_None, LOCTEXT("LODColorationModeDisplayName", "Mesh LODs"));
+					Menu.AddMenuEntry(FEditorViewportCommands::Get().HLODColorationMode, NAME_None, LOCTEXT("HLODColorationModeDisplayName", "Hierarchical LODs"));
+				}
+				Menu.EndSection();
+			}
+		};
+
+		MenuBuilder.AddSubMenu(LOCTEXT("VisualizeGroupedLODDisplayName", "Level of Detail Coloration"), LOCTEXT("GroupedLODMenu_ToolTip", "Select a mode for LOD Coloration"), FNewMenuDelegate::CreateStatic(&Local::BuildLODMenu, this), /*Default*/false, FSlateIcon(FEditorStyle::GetStyleSetName(), "EditorViewport.GroupLODColorationMode"));
+	}
+
 	MenuBuilder.BeginSection("LevelViewportDeferredRendering", LOCTEXT("DeferredRenderingHeader", "Deferred Rendering") );
 	MenuBuilder.EndSection();
-
-	MenuBuilder.AddSubMenu( LOCTEXT("VisualizeBufferViewModeDisplayName", "Buffer Visualization"), LOCTEXT("BufferVisualizationMenu_ToolTip", "Select a mode for buffer visualization"), FNewMenuDelegate::CreateStatic( &BuildBufferVisualizationMenu ), /*Default*/false, FSlateIcon(FEditorStyle::GetStyleSetName(), "EditorViewport.VisualizeBufferMode" ));
+	
+	MenuBuilder.AddSubMenu(LOCTEXT("VisualizeBufferViewModeDisplayName", "Buffer Visualization"), LOCTEXT("BufferVisualizationMenu_ToolTip", "Select a mode for buffer visualization"), FNewMenuDelegate::CreateStatic(&BuildBufferVisualizationMenu), /*Default*/false, FSlateIcon(FEditorStyle::GetStyleSetName(), "EditorViewport.VisualizeBufferMode"));
 
 	MenuBuilder.BeginSection("LevelViewportCollision", LOCTEXT("CollisionViewModeHeader", "Collision") );
 	{
diff --git a/Engine/Source/Editor/Localization/Private/SCulturePicker.cpp b/Engine/Source/Editor/Localization/Private/SCulturePicker.cpp
--- a/Engine/Source/Editor/Localization/Private/SCulturePicker.cpp
+++ b/Engine/Source/Editor/Localization/Private/SCulturePicker.cpp
@@ -66,7 +66,7 @@ void SCulturePicker::Construct( const FArguments& InArgs )
 	const TSharedPtr<FCultureEntry>* InitialSelection = RootEntries.FindByPredicate(IsInitialSelection);
 	if (InitialSelection)
 	{
-		TGuardValue<bool> SupressSelectionGuard(SupressSelectionCallback, true);
+		TGuardValue<bool> SuppressSelectionGuard(SuppressSelectionCallback, true);
 		TreeView->SetSelection(*InitialSelection);
 	}
 }
@@ -279,7 +279,7 @@ void SCulturePicker::OnGetChildren(TSharedPtr<FCultureEntry> Entry, TArray< TSha
 
 void SCulturePicker::OnSelectionChanged(TSharedPtr<FCultureEntry> Entry, ESelectInfo::Type SelectInfo)
 {
-	if (SupressSelectionCallback)
+	if (SuppressSelectionCallback)
 	{
 		return;
 	}
diff --git a/Engine/Source/Editor/Localization/Public/SCulturePicker.h b/Engine/Source/Editor/Localization/Public/SCulturePicker.h
--- a/Engine/Source/Editor/Localization/Public/SCulturePicker.h
+++ b/Engine/Source/Editor/Localization/Public/SCulturePicker.h
@@ -64,7 +64,7 @@ public:
 	SCulturePicker()
 	: DisplayNameFormat(ECultureDisplayFormat::ActiveCultureDisplayName)
 	, CanSelectNone(false)
-	, SupressSelectionCallback(false)
+	, SuppressSelectionCallback(false)
 	{}
 
 	void Construct( const FArguments& InArgs );
@@ -109,5 +109,5 @@ private:
 	bool CanSelectNone;
 
 	/* Flags that the selection callback shouldn't be called when selecting - useful for initial selection, for instance. */
-	bool SupressSelectionCallback;
+	bool SuppressSelectionCallback;
 };
\ No newline at end of file
diff --git a/Engine/Source/Editor/MergeActors/Private/MeshProxyTool/SMeshProxyDialog.cpp b/Engine/Source/Editor/MergeActors/Private/MeshProxyTool/SMeshProxyDialog.cpp
--- a/Engine/Source/Editor/MergeActors/Private/MeshProxyTool/SMeshProxyDialog.cpp
+++ b/Engine/Source/Editor/MergeActors/Private/MeshProxyTool/SMeshProxyDialog.cpp
@@ -247,117 +247,6 @@ void SMeshProxyDialog::CreateLayout()
 					]
 				]
 
-				+SVerticalBox::Slot()
-				.AutoHeight()
-				.Padding(FEditorStyle::GetMargin("StandardDialog.ContentPadding"))
-				[
-					SNew(SCheckBox)
-					.Type(ESlateCheckBoxType::CheckBox)
-					.IsChecked(this, &SMeshProxyDialog::GetUseClippingPlane)
-					.OnCheckStateChanged(this, &SMeshProxyDialog::SetUseClippingPlane)
-					.Content()
-					[
-						SNew(STextBlock)
-						.Text(LOCTEXT("ClippingPlaneLabel", "Use Clipping Plane"))
-						.Font(FEditorStyle::GetFontStyle("StandardDialog.SmallFont"))
-						.ToolTipText(GetPropertyToolTipText(GET_MEMBER_NAME_CHECKED(FMeshProxySettings, bUseClippingPlane)))
-					]
-				]
-
-				+SVerticalBox::Slot()
-				.AutoHeight()
-				.Padding(FEditorStyle::GetMargin("StandardDialog.ContentPadding"))
-				[
-					SNew(SHorizontalBox)
-					+SHorizontalBox::Slot()
-					.VAlign(VAlign_Center)
-					.FillWidth(0.5f)
-					[
-						SNew(SHorizontalBox)
-						+SHorizontalBox::Slot()
-						.AutoWidth()
-						[
-							SNew(SBox)
-							.MinDesiredWidth(30)
-							[
-								SNullWidget::NullWidget
-							]
-						]
-						+SHorizontalBox::Slot()
-						.AutoWidth()
-						.HAlign(HAlign_Left)
-						[
-							SNew(STextBlock)
-							.IsEnabled(this, &SMeshProxyDialog::UseClippingPlaneEnabled)
-							.Text(LOCTEXT("ClippingAxisLabel", "Clipping Axis"))
-							.Font(FEditorStyle::GetFontStyle("StandardDialog.SmallFont"))
-							.ToolTipText(GetPropertyToolTipText(GET_MEMBER_NAME_CHECKED(FMeshProxySettings, AxisIndex)))
-						]
-					]
-					+SHorizontalBox::Slot()
-					.FillWidth(0.5f)
-					.HAlign(HAlign_Left)
-					.VAlign(VAlign_Center)
-					.Padding(8.0, 0.0, 8.0, 0.0)
-					[
-						SNew(STextComboBox)
-						.Font(FEditorStyle::GetFontStyle("StandardDialog.SmallFont"))
-						.OptionsSource(&CuttingPlaneOptions)
-						.InitiallySelectedItem(CuttingPlaneOptions[0])
-						.OnSelectionChanged(this, &SMeshProxyDialog::SetClippingAxis)
-						.IsEnabled(this, &SMeshProxyDialog::UseClippingPlaneEnabled)
-					]
-				]
-
-				+SVerticalBox::Slot()
-				.AutoHeight()
-				.Padding(FEditorStyle::GetMargin("StandardDialog.ContentPadding"))
-				[
-					SNew(SHorizontalBox)
-					+SHorizontalBox::Slot()
-					.FillWidth(0.5f)
-					.VAlign(VAlign_Center)
-					.Padding(0.0, 0.0, 3.0, 0.0)
-					[
-						SNew(SHorizontalBox)
-						+SHorizontalBox::Slot()
-						.AutoWidth()
-						[
-							SNew(SBox)
-							.MinDesiredWidth(30)
-							[
-								SNullWidget::NullWidget
-							]
-						]
-						+SHorizontalBox::Slot()
-						.HAlign(HAlign_Left)
-						.AutoWidth()
-						[
-							SNew(STextBlock)
-							.Text(LOCTEXT("PlaneLevelLabel", "Plane level"))
-							.Font(FEditorStyle::GetFontStyle("StandardDialog.SmallFont"))
-							.IsEnabled(this, &SMeshProxyDialog::UseClippingPlaneEnabled)
-							.ToolTipText(GetPropertyToolTipText(GET_MEMBER_NAME_CHECKED(FMeshProxySettings, ClippingLevel)))
-						]
-					]
-					+SHorizontalBox::Slot()
-					.FillWidth(0.5f)
-					.HAlign(HAlign_Left)
-					[
-						SNew(SBox)
-						.HAlign(HAlign_Fill)
-						.MinDesiredWidth(100.0f)
-						.MaxDesiredWidth(100.0f)
-						[
-							SNew(SNumericEntryBox<float>)
-							.Font(FEditorStyle::GetFontStyle("StandardDialog.SmallFont"))
-							.Value(this, &SMeshProxyDialog::GetClippingLevel)
-							.OnValueChanged(this, &SMeshProxyDialog::ClippingLevelChanged)
-							.IsEnabled(this, &SMeshProxyDialog::UseClippingPlaneEnabled)
-						]
-					]
-				]
-
 				+SVerticalBox::Slot()
 				.AutoHeight()
 				.Padding(FEditorStyle::GetMargin("StandardDialog.ContentPadding"))
@@ -504,55 +393,6 @@ void SMeshProxyDialog::MergeDistanceChanged(int32 NewValue)
 	Tool->ProxySettings.MergeDistance = NewValue;
 }
 
-//Clipping Plane
-ECheckBoxState SMeshProxyDialog::GetUseClippingPlane() const
-{
-	return Tool->ProxySettings.bUseClippingPlane ? ECheckBoxState::Checked: ECheckBoxState::Unchecked;
-}
-
-void SMeshProxyDialog::SetUseClippingPlane(ECheckBoxState NewValue)
-{
-	Tool->ProxySettings.bUseClippingPlane = (NewValue == ECheckBoxState::Checked);
-}
-
-bool SMeshProxyDialog::UseClippingPlaneEnabled() const
-{
-	if(Tool->ProxySettings.bUseClippingPlane)
-	{
-		return true;
-	}
-	return false;
-}
-
-void SMeshProxyDialog::SetClippingAxis(TSharedPtr<FString> NewSelection, ESelectInfo::Type SelectInfo )
-{
-	//This is a ugly hack, but it solves the problem for now
-	int32 AxisIndex = CuttingPlaneOptions.Find(NewSelection);
-	//FMessageDialog::Debugf(*FString::Printf(TEXT("%d"), AxisIndex ));
-	if(AxisIndex < 3)
-	{
-		Tool->ProxySettings.bPlaneNegativeHalfspace = false;
-		Tool->ProxySettings.AxisIndex = AxisIndex;
-		return;
-	}
-	else
-	{
-		Tool->ProxySettings.bPlaneNegativeHalfspace = true;
-		Tool->ProxySettings.AxisIndex = AxisIndex - 3;
-		return;
-	}
-}
-
-TOptional<float> SMeshProxyDialog::GetClippingLevel() const
-{
-	return Tool->ProxySettings.ClippingLevel;
-}
-
-void SMeshProxyDialog::ClippingLevelChanged(float NewValue)
-{
-	Tool->ProxySettings.ClippingLevel = NewValue;
-}
-
 //Texture Resolution
 void SMeshProxyDialog::SetTextureResolution(TSharedPtr<FString> NewSelection, ESelectInfo::Type SelectInfo)
 {
diff --git a/Engine/Source/Editor/MergeActors/Private/MeshProxyTool/SMeshProxyDialog.h b/Engine/Source/Editor/MergeActors/Private/MeshProxyTool/SMeshProxyDialog.h
--- a/Engine/Source/Editor/MergeActors/Private/MeshProxyTool/SMeshProxyDialog.h
+++ b/Engine/Source/Editor/MergeActors/Private/MeshProxyTool/SMeshProxyDialog.h
@@ -37,14 +37,6 @@ protected:
 	TOptional<int32> GetMergeDistance() const;
 	void MergeDistanceChanged(int32 NewValue);
 
-	/** Clipping Plane accessors */
-	ECheckBoxState GetUseClippingPlane() const;
-	void SetUseClippingPlane(ECheckBoxState NewValue);
-	bool UseClippingPlaneEnabled() const;
-	void SetClippingAxis(TSharedPtr<FString> NewSelection, ESelectInfo::Type SelectInfo);
-	TOptional<float> GetClippingLevel() const;
-	void ClippingLevelChanged(float NewValue);
-
 	/** TextureResolution accessors */
 	void SetTextureResolution(TSharedPtr<FString> NewSelection, ESelectInfo::Type SelectInfo);
 	void SetLightMapResolution(TSharedPtr<FString> NewSelection, ESelectInfo::Type SelectInfo);
diff --git a/Engine/Source/Editor/Persona/Private/Persona.cpp b/Engine/Source/Editor/Persona/Private/Persona.cpp
--- a/Engine/Source/Editor/Persona/Private/Persona.cpp
+++ b/Engine/Source/Editor/Persona/Private/Persona.cpp
@@ -2405,17 +2405,17 @@ void FPersona::RemoveAttachedObjectFromPreviewComponent(UObject* Object, FName A
 		PreviewComponent->SetFlags(RF_Transactional);
 		PreviewComponent->Modify();
 
-		for (int32 I=PreviewComponent->AttachChildren.Num()-1; I >= 0; --I) // Iterate backwards because Cleancomponent will remove from AttachChildren
+		for (int32 I=PreviewComponent->GetAttachChildren().Num()-1; I >= 0; --I) // Iterate backwards because CleanupComponent will remove from AttachChildren
 		{
-			USceneComponent* ChildComponent = PreviewComponent->AttachChildren[I];
+			USceneComponent* ChildComponent = PreviewComponent->GetAttachChildren()[I];
 			UObject* Asset = FComponentAssetBrokerage::GetAssetFromComponent(ChildComponent);
 
-			if( Asset == Object && ChildComponent->AttachSocketName == AttachedTo)
+			if( Asset == Object && ChildComponent->GetAttachSocketName() == AttachedTo)
 			{
 				// PreviewComponet will be cleaned up by PreviewScene, 
 				// but if anything is attached, it won't be cleaned up, 
 				// so we'll need to clean them up manually
-				CleanupComponent(PreviewComponent->AttachChildren[I]);
+				CleanupComponent(PreviewComponent->GetAttachChildren()[I]);
 				break;
 			}
 		}
@@ -2426,12 +2426,11 @@ USceneComponent* FPersona::GetComponentForAttachedObject(UObject* Object, FName
 {
 	if (PreviewComponent)
 	{
-		for (int32 I=0; I < PreviewComponent->AttachChildren.Num(); ++I)
+		for (USceneComponent* ChildComponent : PreviewComponent->GetAttachChildren())
 		{
-			USceneComponent* ChildComponent = PreviewComponent->AttachChildren[I];
 			UObject* Asset = FComponentAssetBrokerage::GetAssetFromComponent(ChildComponent);
 
-			if( Asset == Object && ChildComponent->AttachSocketName == AttachedTo)
+			if( Asset == Object && ChildComponent->GetAttachSocketName() == AttachedTo)
 			{
 				return ChildComponent;
 			}
@@ -2463,9 +2462,9 @@ void FPersona::RemoveAttachedComponent( bool bRemovePreviewAttached /* = true */
 	// clean up components	
 	if (PreviewComponent)
 	{
-		for (int32 I=PreviewComponent->AttachChildren.Num()-1; I >= 0; --I) // Iterate backwards because Cleancomponent will remove from AttachChildren
+		for (int32 I=PreviewComponent->GetAttachChildren().Num()-1; I >= 0; --I) // Iterate backwards because CleanupComponent will remove from AttachChildren
 		{
-			USceneComponent* ChildComponent = PreviewComponent->AttachChildren[I];
+			USceneComponent* ChildComponent = PreviewComponent->GetAttachChildren()[I];
 			UObject* Asset = FComponentAssetBrokerage::GetAssetFromComponent(ChildComponent);
 
 			bool bRemove = true;
@@ -2476,7 +2475,7 @@ void FPersona::RemoveAttachedComponent( bool bRemovePreviewAttached /* = true */
 				//could this asset have come from the skeleton
 				if(PreviewAttachedObjects.Contains(Asset))
 				{
-					if(PreviewAttachedObjects.Find(Asset)->Contains(ChildComponent->AttachSocketName))
+					if(PreviewAttachedObjects.Find(Asset)->Contains(ChildComponent->GetAttachSocketName()))
 					{
 						bRemove = false;
 					}
@@ -2488,13 +2487,13 @@ void FPersona::RemoveAttachedComponent( bool bRemovePreviewAttached /* = true */
 				// PreviewComponet will be cleaned up by PreviewScene, 
 				// but if anything is attached, it won't be cleaned up, 
 				// so we'll need to clean them up manually
-				CleanupComponent(PreviewComponent->AttachChildren[I]);
+				CleanupComponent(PreviewComponent->GetAttachChildren()[I]);
 			}
 		}
 
 		if( bRemovePreviewAttached )
 		{
-			PreviewComponent->AttachChildren.Empty();
+			check(PreviewComponent->GetAttachChildren().Num() == 0);
 		}
 	}
 }
@@ -2503,12 +2502,12 @@ void FPersona::CleanupComponent(USceneComponent* Component)
 {
 	if (Component)
 	{
-		for (int32 I=0; I<Component->AttachChildren.Num(); ++I)
+		for (int32 I = Component->GetAttachChildren().Num() - 1; I >= 0; --I) // Iterate backwards because CleanupComponent will remove from AttachChildren
 		{
-			CleanupComponent(Component->AttachChildren[I]);
+			CleanupComponent(Component->GetAttachChildren()[I]);
 		}
 
-		Component->AttachChildren.Empty();
+		check(Component->GetAttachChildren().Num() == 0);
 		Component->DestroyComponent();
 	}
 }
diff --git a/Engine/Source/Editor/Persona/Private/SAnimationSequenceBrowser.cpp b/Engine/Source/Editor/Persona/Private/SAnimationSequenceBrowser.cpp
--- a/Engine/Source/Editor/Persona/Private/SAnimationSequenceBrowser.cpp
+++ b/Engine/Source/Editor/Persona/Private/SAnimationSequenceBrowser.cpp
@@ -47,15 +47,15 @@ SAnimationSequenceBrowser::~SAnimationSequenceBrowser()
 {
 	if(PreviewComponent)
 	{
-		for(int32 ComponentIdx = PreviewComponent->AttachChildren.Num() - 1 ; ComponentIdx >= 0 ; --ComponentIdx)
+		for(int32 ComponentIdx = PreviewComponent->GetAttachChildren().Num() - 1 ; ComponentIdx >= 0 ; --ComponentIdx)
 		{
-			USceneComponent* Component = PreviewComponent->AttachChildren[ComponentIdx];
+			USceneComponent* Component = PreviewComponent->GetAttachChildren()[ComponentIdx];
 			if(Component)
 			{
 				CleanupPreviewSceneComponent(Component);
 			}
 		}
-		PreviewComponent->AttachChildren.Empty();
+		check(PreviewComponent->GetAttachChildren().Num() == 0);
 	}
 
 	if(ViewportClient.IsValid())
@@ -985,12 +985,12 @@ void SAnimationSequenceBrowser::CleanupPreviewSceneComponent(USceneComponent* Co
 {
 	if(Component)
 	{
-		for(int32 ComponentIdx = Component->AttachChildren.Num() - 1 ; ComponentIdx >= 0 ; --ComponentIdx)
+		for(int32 ComponentIdx = Component->GetAttachChildren().Num() - 1 ; ComponentIdx >= 0 ; --ComponentIdx)
 		{
-			USceneComponent* ChildComponent = Component->AttachChildren[ComponentIdx];
+			USceneComponent* ChildComponent = Component->GetAttachChildren()[ComponentIdx];
 			CleanupPreviewSceneComponent(ChildComponent);
 		}
-		Component->AttachChildren.Empty();
+		check(Component->GetAttachChildren().Num() == 0);
 		Component->DestroyComponent();
 	}
 }
diff --git a/Engine/Source/Editor/SceneOutliner/Private/ActorTreeItem.cpp b/Engine/Source/Editor/SceneOutliner/Private/ActorTreeItem.cpp
--- a/Engine/Source/Editor/SceneOutliner/Private/ActorTreeItem.cpp
+++ b/Engine/Source/Editor/SceneOutliner/Private/ActorTreeItem.cpp
@@ -181,9 +181,9 @@ void FActorDropTarget::PerformAttachment(FName SocketName, TWeakObjectPtr<AActor
 void FActorDropTarget::DetachActorFromParent(AActor* ChildActor)
 {
 	USceneComponent* RootComp = ChildActor->GetRootComponent();
-	if (RootComp && RootComp->AttachParent)
+	if (RootComp && RootComp->GetAttachParent())
 	{
-		AActor* OldParent = RootComp->AttachParent->GetOwner();
+		AActor* OldParent = RootComp->GetAttachParent()->GetOwner();
 		OldParent->Modify();
 		RootComp->DetachFromParent(true);
 		
diff --git a/Engine/Source/Editor/SceneOutliner/Private/FolderTreeItem.cpp b/Engine/Source/Editor/SceneOutliner/Private/FolderTreeItem.cpp
--- a/Engine/Source/Editor/SceneOutliner/Private/FolderTreeItem.cpp
+++ b/Engine/Source/Editor/SceneOutliner/Private/FolderTreeItem.cpp
@@ -152,9 +152,9 @@ void FFolderDropTarget::OnDrop(FDragDropPayload& DraggedObjects, UWorld& World,
 			auto* RootComp = Parent->GetRootComponent();
 
 			// We don't detach if it's a child of another that's been dragged
-			if (RootComp && RootComp->AttachParent && !ChildActors.Contains(Parent))
+			if (RootComp && RootComp->GetAttachParent() && !ChildActors.Contains(Parent))
 			{
-				if (AActor* OldParentActor = RootComp->AttachParent->GetOwner())
+				if (AActor* OldParentActor = RootComp->GetAttachParent()->GetOwner())
 				{
 					OldParentActor->Modify();
 				}
diff --git a/Engine/Source/Editor/Sequencer/Private/SequencerEdMode.cpp b/Engine/Source/Editor/Sequencer/Private/SequencerEdMode.cpp
--- a/Engine/Source/Editor/Sequencer/Private/SequencerEdMode.cpp
+++ b/Engine/Source/Editor/Sequencer/Private/SequencerEdMode.cpp
@@ -99,9 +99,9 @@ FTransform GetRefFrame(const UObject* InObject)
 	FTransform RefTM = FTransform::Identity;
 
 	const AActor* Actor = Cast<AActor>(InObject);
-	if (Actor != nullptr && Actor->GetRootComponent() != nullptr && Actor->GetRootComponent()->AttachParent != nullptr)
+	if (Actor != nullptr && Actor->GetRootComponent() != nullptr && Actor->GetRootComponent()->GetAttachParent() != nullptr)
 	{
-		RefTM = Actor->GetRootComponent()->AttachParent->GetSocketTransform(Actor->GetRootComponent()->AttachSocketName);
+		RefTM = Actor->GetRootComponent()->GetAttachParent()->GetSocketTransform(Actor->GetRootComponent()->GetAttachSocketName());
 	}
 
 	return RefTM;
diff --git a/Engine/Source/Editor/Sequencer/Private/TimeSliderController.cpp b/Engine/Source/Editor/Sequencer/Private/TimeSliderController.cpp
--- a/Engine/Source/Editor/Sequencer/Private/TimeSliderController.cpp
+++ b/Engine/Source/Editor/Sequencer/Private/TimeSliderController.cpp
@@ -107,7 +107,7 @@ FSequencerTimeSliderController::FSequencerTimeSliderController( const FTimeSlide
 	ScrubHandleUp = FEditorStyle::GetBrush( TEXT( "Sequencer.Timeline.ScrubHandleUp" ) ); 
 	ScrubHandleDown = FEditorStyle::GetBrush( TEXT( "Sequencer.Timeline.ScrubHandleDown" ) );
 	ScrubHandleSize = 13.f;
-	ContextMenuSupression = 0;
+	ContextMenuSuppression = 0;
 }
 
 struct FDrawTickArgs
@@ -472,7 +472,7 @@ FReply FSequencerTimeSliderController::OnMouseButtonUp( SWidget& WidgetOwner, co
 		if (!bPanning)
 		{
 			// Open a context menu if allowed
-			if (ContextMenuSupression == 0 && TimeSliderArgs.PlaybackRange.IsSet())
+			if (ContextMenuSuppression == 0 && TimeSliderArgs.PlaybackRange.IsSet())
 			{
 				FScrubRangeToScreen RangeToScreen( TimeSliderArgs.ViewRange.Get(), MyGeometry.Size );
 				FVector2D CursorPos = MyGeometry.AbsoluteToLocal( MouseEvent.GetScreenSpacePosition() );
diff --git a/Engine/Source/Editor/Sequencer/Private/TimeSliderController.h b/Engine/Source/Editor/Sequencer/Private/TimeSliderController.h
--- a/Engine/Source/Editor/Sequencer/Private/TimeSliderController.h
+++ b/Engine/Source/Editor/Sequencer/Private/TimeSliderController.h
@@ -179,7 +179,7 @@ private:
 	/** Range stack */
 	TArray<FVector2D> RangeStack;
 	/** When > 0, we should not show context menus */
-	int32 ContextMenuSupression;
+	int32 ContextMenuSuppression;
 	friend FContextMenuSuppressor;
 };
 
@@ -188,11 +188,11 @@ struct FContextMenuSuppressor
 	FContextMenuSuppressor(TSharedRef<FSequencerTimeSliderController> InTimeSliderController)
 		: TimeSliderController(InTimeSliderController)
 	{
-		++TimeSliderController->ContextMenuSupression;
+		++TimeSliderController->ContextMenuSuppression;
 	}
 	~FContextMenuSuppressor()
 	{
-		--TimeSliderController->ContextMenuSupression;
+		--TimeSliderController->ContextMenuSuppression;
 	}
 
 private:
diff --git a/Engine/Source/Editor/StaticMeshEditor/Private/StaticMeshEditorTools.cpp b/Engine/Source/Editor/StaticMeshEditor/Private/StaticMeshEditorTools.cpp
--- a/Engine/Source/Editor/StaticMeshEditor/Private/StaticMeshEditorTools.cpp
+++ b/Engine/Source/Editor/StaticMeshEditor/Private/StaticMeshEditorTools.cpp
@@ -56,7 +56,7 @@ void FStaticMeshDetails::CustomizeDetails( class IDetailLayoutBuilder& DetailBui
 
 	if( HiddenBodyInstanceProps.Num() == 0 )
 	{
-		HiddenBodyInstanceProps.Add("DefaultInstance");
+		//HiddenBodyInstanceProps.Add("DefaultInstance");
 		HiddenBodyInstanceProps.Add("BoneName");
 		HiddenBodyInstanceProps.Add("PhysicsType");
 		HiddenBodyInstanceProps.Add("bConsiderForBounds");
@@ -1370,6 +1370,8 @@ TSharedRef<SWidget> FMeshSectionSettingsLayout::OnGenerateWidgetsForMaterial(UMa
 			.Padding(0,2,0,0)
 		[
 			SNew(SCheckBox)
+				.IsEnabled(this, &FMeshSectionSettingsLayout::SectionCollisionEnabled)
+				.ToolTipText(this, &FMeshSectionSettingsLayout::GetCollisionEnabledToolTip)
 				.IsChecked(this, &FMeshSectionSettingsLayout::DoesSectionCollide, SlotIndex)
 				.OnCheckStateChanged(this, &FMeshSectionSettingsLayout::OnSectionCollisionChanged, SlotIndex)
 			[
@@ -1412,6 +1414,29 @@ void FMeshSectionSettingsLayout::OnSectionCastShadowChanged(ECheckBoxState NewSt
 	CallPostEditChange();
 }
 
+bool FMeshSectionSettingsLayout::SectionCollisionEnabled() const
+{
+	UStaticMesh& StaticMesh = GetStaticMesh();
+	// Only enable 'Enable Collision' check box if this LOD is used for collision
+	return (StaticMesh.LODForCollision == LODIndex);
+}
+
+FText FMeshSectionSettingsLayout::GetCollisionEnabledToolTip() const
+{
+	UStaticMesh& StaticMesh = GetStaticMesh();
+	
+	// If using a different LOD for collision, disable the check box
+	if (StaticMesh.LODForCollision != LODIndex)
+	{
+		return LOCTEXT("EnableCollisionToolTipDisabled", "This LOD is not used for collision, see the LODForCollision setting.");
+	}
+	// This LOD is used for collision, give info on what flag does
+	else
+	{
+		return LOCTEXT("EnableCollisionToolTipEnabled", "Controls whether this section ever has per-poly collision. Disabling this where possible will lower memory usage for this mesh.");
+	}
+}
+
 ECheckBoxState FMeshSectionSettingsLayout::DoesSectionCollide(int32 SectionIndex) const
 {
 	UStaticMesh& StaticMesh = GetStaticMesh();
diff --git a/Engine/Source/Editor/StaticMeshEditor/Private/StaticMeshEditorTools.h b/Engine/Source/Editor/StaticMeshEditor/Private/StaticMeshEditorTools.h
--- a/Engine/Source/Editor/StaticMeshEditor/Private/StaticMeshEditorTools.h
+++ b/Engine/Source/Editor/StaticMeshEditor/Private/StaticMeshEditorTools.h
@@ -243,6 +243,8 @@ private:
 	ECheckBoxState DoesSectionCastShadow(int32 SectionIndex) const;
 	void OnSectionCastShadowChanged(ECheckBoxState NewState, int32 SectionIndex);
 	ECheckBoxState DoesSectionCollide(int32 SectionIndex) const;
+	bool SectionCollisionEnabled() const;
+	FText GetCollisionEnabledToolTip() const;
 	void OnSectionCollisionChanged(ECheckBoxState NewState, int32 SectionIndex);
 	ECheckBoxState IsSectionHighlighted(int32 SectionIndex) const;
 	void OnSectionHighlightedChanged(ECheckBoxState NewState, int32 SectionIndex);
diff --git a/Engine/Source/Editor/StaticMeshEditor/Private/StaticMeshEditorViewportClient.cpp b/Engine/Source/Editor/StaticMeshEditor/Private/StaticMeshEditorViewportClient.cpp
--- a/Engine/Source/Editor/StaticMeshEditor/Private/StaticMeshEditorViewportClient.cpp
+++ b/Engine/Source/Editor/StaticMeshEditor/Private/StaticMeshEditorViewportClient.cpp
@@ -76,7 +76,7 @@ FStaticMeshEditorViewportClient::FStaticMeshEditorViewportClient(TWeakPtr<IStati
 	OverrideNearClipPlane(1.0f);
 	bUsingOrbitCamera = true;
 
-	bShowCollision = true;
+	bShowCollision = false;
 	bShowSockets = true;
 	bDrawUVs = false;
 	bDrawNormals = false;
@@ -792,8 +792,8 @@ void FStaticMeshEditorViewportClient::DrawCanvas( FViewport& InViewport, FSceneV
 		FText::AsNumber(int32(StaticMesh->GetBounds().BoxExtent.Y * 2.0f)),
 		FText::AsNumber(int32(StaticMesh->GetBounds().BoxExtent.Z * 2.0f)))));
 
-	// Show the number of collision primitives if we are drawing collision.
-	if(bShowCollision && StaticMesh->BodySetup)
+	// Show the number of collision primitives
+	if(StaticMesh->BodySetup)
 	{
 		TextItems.Add(SStaticMeshEditorViewport::FOverlayTextItem(
 			FText::Format(NSLOCTEXT("UnrealEd", "NumPrimitives_F", "Num Collision Primitives:  {0}"), FText::AsNumber(StaticMesh->BodySetup->AggGeom.GetElementCount()))));
@@ -1187,6 +1187,12 @@ void FStaticMeshEditorViewportClient::SetPreviewMesh(UStaticMesh* InStaticMesh,
 	StaticMesh = InStaticMesh;
 	StaticMeshComponent = InStaticMeshComponent;
 
+	if(StaticMeshComponent != nullptr)
+	{
+		StaticMeshComponent->bDrawMeshCollisionWireframe = bShowCollision;
+		StaticMeshComponent->MarkRenderStateDirty();
+	}
+
 	// If we have a thumbnail transform, we will favor that over the camera position as the user may have customized this for a nice view
 	// If we have neither a custom thumbnail nor a valid camera position, then we'll just use the default thumbnail transform 
 	const USceneThumbnailInfo* const AssetThumbnailInfo = Cast<USceneThumbnailInfo>(StaticMesh->ThumbnailInfo);
@@ -1305,6 +1311,13 @@ bool FStaticMeshEditorViewportClient::IsSetDrawVerticesChecked() const
 void FStaticMeshEditorViewportClient::SetShowWireframeCollision()
 {
 	bShowCollision = !bShowCollision;
+
+	if(StaticMeshComponent != nullptr)
+	{
+		StaticMeshComponent->bDrawMeshCollisionWireframe = bShowCollision;
+		StaticMeshComponent->MarkRenderStateDirty();
+	}
+
 	if (FEngineAnalytics::IsAvailable())
 	{
 		FEngineAnalytics::GetProvider().RecordEvent(TEXT("Editor.Usage.StaticMesh.Toolbar"), TEXT("bShowCollision"), bShowCollision ? TEXT("True") : TEXT("False"));
diff --git a/Engine/Source/Editor/UnrealEd/Classes/Editor/EditorEngine.h b/Engine/Source/Editor/UnrealEd/Classes/Editor/EditorEngine.h
--- a/Engine/Source/Editor/UnrealEd/Classes/Editor/EditorEngine.h
+++ b/Engine/Source/Editor/UnrealEd/Classes/Editor/EditorEngine.h
@@ -635,6 +635,63 @@ public:
 	DECLARE_EVENT_OneParam( UEditorEngine, FGetActorRecordingState, bool& /* bIsRecordingActive */ );
 	FGetActorRecordingState& GetActorRecordingState() { return GetActorRecordingStateEvent; }
 
+
+
+
+
+
+	/** Editor-only event triggered when a HLOD Actor is moved between clusters */
+	DECLARE_EVENT_TwoParams(UEngine, FHLODActorMovedEvent, const AActor*, const AActor*);
+	FHLODActorMovedEvent& OnHLODActorMoved() { return HLODActorMovedEvent; }
+
+	/** Called by internal engine systems after a HLOD Actor is moved between clusters */
+	void BroadcastHLODActorMoved(const AActor* InActor, const AActor* ParentActor) { HLODActorMovedEvent.Broadcast(InActor, ParentActor); }
+
+	/** Editor-only event triggered when a HLOD Actor's mesh is build */
+	DECLARE_EVENT_OneParam(UEngine, FHLODMeshBuildEvent, const class ALODActor*);
+	FHLODMeshBuildEvent& OnHLODMeshBuild() { return HLODMeshBuildEvent; }
+
+	/** Called by internal engine systems after a HLOD Actor's mesh is build */
+	void BroadcastHLODMeshBuild(const class ALODActor* InActor) { HLODMeshBuildEvent.Broadcast(InActor); }
+
+	/** Editor-only event triggered when a HLOD Actor is added to a cluster */
+	DECLARE_EVENT_TwoParams(UEngine, FHLODActorAddedEvent, const AActor*, const AActor*);
+	FHLODActorAddedEvent& OnHLODActorAdded() { return HLODActorAddedEvent; }
+
+	/** Called by internal engine systems after a HLOD Actor is added to a cluster */
+	void BroadcastHLODActorAdded(const AActor* InActor, const AActor* ParentActor) { HLODActorAddedEvent.Broadcast(InActor, ParentActor); }
+
+	/** Editor-only event triggered when a HLOD Actor is marked dirty */
+	DECLARE_EVENT_OneParam(UEngine, FHLODActorMarkedDirtyEvent, class ALODActor*);
+	FHLODActorMarkedDirtyEvent& OnHLODActorMarkedDirty() { return HLODActorMarkedDirtyEvent; }
+
+	/** Called by internal engine systems after a HLOD Actor is marked dirty */
+	void BroadcastHLODActorMarkedDirty(class ALODActor* InActor) { HLODActorMarkedDirtyEvent.Broadcast(InActor); }
+
+	/** Editor-only event triggered when a HLOD Actor is marked dirty */
+	DECLARE_EVENT(UEngine, FHLODTransitionScreenSizeChangedEvent);
+	FHLODTransitionScreenSizeChangedEvent& OnHLODTransitionScreenSizeChanged() { return HLODTransitionScreenSizeChangedEvent; }
+
+	/** Called by internal engine systems after a HLOD Actor is marked dirty */
+	void BroadcastHLODTransitionScreenSizeChanged() { HLODTransitionScreenSizeChangedEvent.Broadcast(); }
+
+	/** Editor-only event triggered when a HLOD level is added or removed */
+	DECLARE_EVENT(UEngine, FHLODLevelsArrayChangedEvent);
+	FHLODLevelsArrayChangedEvent& OnHLODLevelsArrayChanged() { return HLODLevelsArrayChangedEvent; }
+
+	/** Called by internal engine systems after a HLOD Actor is marked dirty */
+	void BroadcastHLODLevelsArrayChanged() { HLODLevelsArrayChangedEvent.Broadcast(); }
+
+	DECLARE_EVENT_TwoParams(UEngine, FHLODActorRemovedFromClusterEvent, const AActor*, const AActor*);
+	FHLODActorRemovedFromClusterEvent& OnHLODActorRemovedFromCluster() { return HLODActorRemovedFromClusterEvent; }
+
+	/** Called by internal engine systems after an Actor is removed from a cluster */
+	void BroadcastHLODActorRemovedFromCluster(const AActor* InActor, const AActor* ParentActor) { HLODActorRemovedFromClusterEvent.Broadcast(InActor, ParentActor); }
+
+
+
+
+
 	/**
 	* Called before an actor or component is about to be translated, rotated, or scaled by the editor
 	*
@@ -2595,6 +2652,27 @@ private:
 
 	/** Delegate broadcast when the camera viewed through the viewport has been moved */
 	FOnEndTransformCamera OnEndCameraTransformEvent;
+	
+	/** Broadcasts after an HLOD actor has been moved between clusters */
+	FHLODActorMovedEvent HLODActorMovedEvent;
+
+	/** Broadcasts after an HLOD actor's mesh is build*/
+	FHLODMeshBuildEvent HLODMeshBuildEvent;
+
+	/** Broadcasts after an HLOD actor has added to a cluster */
+	FHLODActorAddedEvent HLODActorAddedEvent;
+
+	/** Broadcasts after an HLOD actor has been marked dirty */
+	FHLODActorMarkedDirtyEvent HLODActorMarkedDirtyEvent;
+
+	/** Broadcasts after a Draw distance value (World settings) is changed */
+	FHLODTransitionScreenSizeChangedEvent HLODTransitionScreenSizeChangedEvent;
+
+	/** Broadcasts after the HLOD levels array is changed */
+	FHLODLevelsArrayChangedEvent HLODLevelsArrayChangedEvent;
+
+	/** Broadcasts after an Actor is removed from a cluster */
+	FHLODActorRemovedFromClusterEvent HLODActorRemovedFromClusterEvent;
 
 	/** Delegate broadcast by the engine every tick when PIE/SIE is active, to check to see whether we need to
 		be able to capture state for simulating actor (for Sequencer recording features) */
diff --git a/Engine/Source/Editor/UnrealEd/Private/AssetSelection.cpp b/Engine/Source/Editor/UnrealEd/Private/AssetSelection.cpp
--- a/Engine/Source/Editor/UnrealEd/Private/AssetSelection.cpp
+++ b/Engine/Source/Editor/UnrealEd/Private/AssetSelection.cpp
@@ -195,7 +195,7 @@ namespace AssetSelectionUtils
 					}
 
 					USceneComponent* RootComp = CurrentActor->GetRootComponent();
-					if(RootComp != NULL && RootComp->AttachParent != NULL)
+					if(RootComp != nullptr && RootComp->GetAttachParent() != nullptr)
 					{
 						ActorInfo.bHaveAttachedActor = true;
 					}
diff --git a/Engine/Source/Editor/UnrealEd/Private/Dialogs/Dialogs.cpp b/Engine/Source/Editor/UnrealEd/Private/Dialogs/Dialogs.cpp
--- a/Engine/Source/Editor/UnrealEd/Private/Dialogs/Dialogs.cpp
+++ b/Engine/Source/Editor/UnrealEd/Private/Dialogs/Dialogs.cpp
@@ -677,7 +677,7 @@ FSuppressableWarningDialog::FSuppressableWarningDialog(const FSetupInfo& Info)
 			.bHasCancelButton(!Info.CancelText.IsEmpty())
 			.ConfirmText(Info.ConfirmText)
 			.CancelText(Info.CancelText)
-			.bDefaultCheckValue(Info.bDefaultToSupressInTheFuture)
+			.bDefaultCheckValue(Info.bDefaultToSuppressInTheFuture)
 			.CheckboxMessage(Info.CheckBoxText)
 			.ParentWindow(ModalWindow)
 			.Image((Info.Image != NULL) ? Info.Image : DefaultImage);
@@ -711,7 +711,7 @@ FSuppressableWarningDialog::EResult FSuppressableWarningDialog::ShowModal() cons
 	else
 	{
 		// If the dialog is suppressed, log the warning
-		UE_LOG(LogDialogs, Warning, TEXT("Supressed: %s"), *Prompt.ToString());
+		UE_LOG(LogDialogs, Warning, TEXT("Suppressed: %s"), *Prompt.ToString());
 	}
 
 	return RetCode;
diff --git a/Engine/Source/Editor/UnrealEd/Private/EdMode.cpp b/Engine/Source/Editor/UnrealEd/Private/EdMode.cpp
--- a/Engine/Source/Editor/UnrealEd/Private/EdMode.cpp
+++ b/Engine/Source/Editor/UnrealEd/Private/EdMode.cpp
@@ -62,7 +62,7 @@ namespace
 	 * Returns a reference to the named property value data in the given container.
 	 */
 	template<typename T>
-	T* GetPropertyValuePtrByName(const UStruct* InStruct, void* InContainer, FString PropertyName, int32 ArrayIndex)
+	T* GetPropertyValuePtrByName(const UStruct* InStruct, void* InContainer, FString PropertyName, int32 ArrayIndex, UProperty*& OutProperty)
 	{
 		T* ValuePtr = NULL;
 
@@ -90,7 +90,7 @@ namespace
 			if(StructProp != NULL)
 			{
 				// Recursively call back into this function with the structure property and container value
-				ValuePtr = GetPropertyValuePtrByName<T>(StructProp->Struct, StructProp->ContainerPtrToValuePtr<void>(InContainer), PropertyName.RightChop(DelimPos + 1), ArrayIndex);
+				ValuePtr = GetPropertyValuePtrByName<T>(StructProp->Struct, StructProp->ContainerPtrToValuePtr<void>(InContainer), PropertyName.RightChop(DelimPos + 1), ArrayIndex, OutProperty);
 			}
 			else
 			{
@@ -106,7 +106,7 @@ namespace
 						if(ArrayHelper.IsValidIndex(SubArrayIndex))
 						{
 							// Recursively call back into this function with the array element and container value
-							ValuePtr = GetPropertyValuePtrByName<T>(StructProp->Struct, ArrayHelper.GetRawPtr(SubArrayIndex), PropertyName.RightChop(DelimPos + 1), ArrayIndex);
+							ValuePtr = GetPropertyValuePtrByName<T>(StructProp->Struct, ArrayHelper.GetRawPtr(SubArrayIndex), PropertyName.RightChop(DelimPos + 1), ArrayIndex, OutProperty);
 						}
 					}
 				}
@@ -133,6 +133,8 @@ namespace
 					// Property is a vector property, so access directly
 					ValuePtr = Prop->ContainerPtrToValuePtr<T>(InContainer);
 				}
+
+				OutProperty = Prop;
 			}
 		}
 
@@ -146,7 +148,8 @@ namespace
 	T GetPropertyValueByName(UObject* Object, FString PropertyName, int32 PropertyIndex)
 	{
 		T Value;
-		if (T* ValuePtr = GetPropertyValuePtrByName<T>(Object->GetClass(), Object, PropertyName, PropertyIndex))
+		UProperty* DummyProperty = NULL;
+		if (T* ValuePtr = GetPropertyValuePtrByName<T>(Object->GetClass(), Object, PropertyName, PropertyIndex, DummyProperty))
 		{
 			Value = *ValuePtr;
 		}
@@ -157,9 +160,9 @@ namespace
 	 * Sets the property with the given name in the given Actor instance to the given value.
 	 */
 	template<typename T>
-	void SetPropertyValueByName(UObject* Object, FString PropertyName, int32 PropertyIndex, const T& InValue)
+	void SetPropertyValueByName(UObject* Object, FString PropertyName, int32 PropertyIndex, const T& InValue, UProperty*& OutProperty)
 	{
-		if (T* ValuePtr = GetPropertyValuePtrByName<T>(Object->GetClass(), Object, PropertyName, PropertyIndex))
+		if (T* ValuePtr = GetPropertyValuePtrByName<T>(Object->GetClass(), Object, PropertyName, PropertyIndex, OutProperty))
 		{
 			*ValuePtr = InValue;
 		}
@@ -367,16 +370,20 @@ bool FEdMode::InputDelta(FEditorViewportClient* InViewportClient, FViewport* InV
 
 				BestSelectedItem->PreEditChange(NULL);
 
+				// Property that we actually change
+				UProperty* SetProperty = NULL;
+
 				if (bEditedPropertyIsTransform)
 				{
-					SetPropertyValueByName<FTransform>(BestSelectedItem, EditedPropertyName, EditedPropertyIndex, LocalTM);
+					SetPropertyValueByName<FTransform>(BestSelectedItem, EditedPropertyName, EditedPropertyIndex, LocalTM, SetProperty);
 				}
 				else
 				{
-					SetPropertyValueByName<FVector>(BestSelectedItem, EditedPropertyName, EditedPropertyIndex, LocalTM.GetLocation());
+					SetPropertyValueByName<FVector>(BestSelectedItem, EditedPropertyName, EditedPropertyIndex, LocalTM.GetLocation(), SetProperty);
 				}
 
-				BestSelectedItem->PostEditChange();
+				FPropertyChangedEvent PropertyChangeEvent(SetProperty);
+				BestSelectedItem->PostEditChangeProperty(PropertyChangeEvent);
 
 				return true;
 			}
diff --git a/Engine/Source/Editor/UnrealEd/Private/EditorActor.cpp b/Engine/Source/Editor/UnrealEd/Private/EditorActor.cpp
--- a/Engine/Source/Editor/UnrealEd/Private/EditorActor.cpp
+++ b/Engine/Source/Editor/UnrealEd/Private/EditorActor.cpp
@@ -198,7 +198,7 @@ bool UUnrealEdEngine::WarnIfDestinationLevelIsHidden( UWorld* InWorld )
 {
 	bool result = false;
 	//prepare the warning dialog
-	FSuppressableWarningDialog::FSetupInfo Info( LOCTEXT( "Warning_PasteWarningBody","You are trying to paste to a hidden level.\nSupressing this will default to Do Not Paste" ), LOCTEXT( "Warning_PasteWarningHeader","Pasting To Hidden Level" ), "PasteHiddenWarning" );
+	FSuppressableWarningDialog::FSetupInfo Info( LOCTEXT( "Warning_PasteWarningBody","You are trying to paste to a hidden level.\nSuppressing this will default to Do Not Paste" ), LOCTEXT( "Warning_PasteWarningHeader","Pasting To Hidden Level" ), "PasteHiddenWarning" );
 	Info.ConfirmText = LOCTEXT( "Warning_PasteContinue","Unhide Level and paste" );
 	Info.CancelText = LOCTEXT( "Warning_PasteCancel","Do not paste" );
 	FSuppressableWarningDialog PasteHiddenWarning( Info );
diff --git a/Engine/Source/Editor/UnrealEd/Private/EditorAnimUtils.cpp b/Engine/Source/Editor/UnrealEd/Private/EditorAnimUtils.cpp
--- a/Engine/Source/Editor/UnrealEd/Private/EditorAnimUtils.cpp
+++ b/Engine/Source/Editor/UnrealEd/Private/EditorAnimUtils.cpp
@@ -67,14 +67,16 @@ namespace EditorAnimUtils
 
 		if(bRetargetReferredAssets)
 		{
-			for(auto Iter = ComplexAnimsToRetarget.CreateConstIterator(); Iter; ++Iter)
+			// Grab assets from the blueprint. Do this first as it can add complex assets to the retarget array
+			// which will need to be processed next.
+			for(auto Iter = AnimBlueprintsToRetarget.CreateConstIterator(); Iter; ++Iter)
 			{
-				(*Iter)->GetAllAnimationSequencesReferred(AnimSequencesToRetarget);
+				GetAllAnimationSequencesReferredInBlueprint( (*Iter), ComplexAnimsToRetarget, AnimSequencesToRetarget);
 			}
 
-			for(auto Iter = AnimBlueprintsToRetarget.CreateConstIterator(); Iter; ++Iter)
+			for(auto Iter = ComplexAnimsToRetarget.CreateConstIterator(); Iter; ++Iter)
 			{
-				GetAllAnimationSequencesReferredInBlueprint( (*Iter), ComplexAnimsToRetarget, AnimSequencesToRetarget);
+				(*Iter)->GetAllAnimationSequencesReferred(AnimSequencesToRetarget);
 			}
 
 			int SequenceIndex = 0;
@@ -441,6 +443,27 @@ namespace EditorAnimUtils
 				}
 			}
 		}
+
+		// Pull all the assets we reference from variables
+		TArray<UProperty*> SimpleAnimVariableProperties;
+		TArray<UProperty*> ComplexAnimVariableProperties;
+		TArray<UAnimSequence*> VariableReferencedSimpleAnims;
+		TArray<UAnimationAsset*> VariableReferencedComplexAnims;
+
+		UObject* DefaultObject = AnimBlueprint->GetAnimBlueprintGeneratedClass()->GetDefaultObject();
+		GetBlueprintAssetVariableProperties(AnimBlueprint, SimpleAnimVariableProperties, ComplexAnimVariableProperties);
+		GetAssetsFromProperties(SimpleAnimVariableProperties, DefaultObject, VariableReferencedSimpleAnims);
+		GetAssetsFromProperties(ComplexAnimVariableProperties, DefaultObject, VariableReferencedComplexAnims);
+
+		for(UAnimSequence* Sequence : VariableReferencedSimpleAnims)
+		{
+			AnimSequences.AddUnique(Sequence);
+		}
+
+		for(UAnimationAsset* Asset : VariableReferencedComplexAnims)
+		{
+			ComplexAnims.AddUnique(Asset);
+		}
 	}
 
 	void ReplaceReferredAnimationsInBlueprint(UAnimBlueprint* AnimBlueprint, const TMap<UAnimationAsset*, UAnimationAsset*>& ComplexAnimMap, const TMap<UAnimSequence*, UAnimSequence*>& AnimSequenceMap)
@@ -458,6 +481,82 @@ namespace EditorAnimUtils
 				}
 			}
 		}
+
+		// Push all the assets we reference from variables
+		TArray<UProperty*> SimpleAnimVariableProperties;
+		TArray<UProperty*> ComplexAnimVariableProperties;
+
+		UObject* DefaultObject = AnimBlueprint->GetAnimBlueprintGeneratedClass()->GetDefaultObject();
+		GetBlueprintAssetVariableProperties(AnimBlueprint, SimpleAnimVariableProperties, ComplexAnimVariableProperties);
+
+		for(UProperty* SimpleProp : SimpleAnimVariableProperties)
+		{
+			if(UObject** ResolvedObject = SimpleProp->ContainerPtrToValuePtr<UObject*>(DefaultObject))
+			{
+				if(UAnimSequence* Sequence = Cast<UAnimSequence>(*ResolvedObject))
+				{
+					// Found an anim sequence variable that's set to a valid sequence
+					if(UAnimSequence* const* NewSequence = AnimSequenceMap.Find(Sequence))
+					{
+						*ResolvedObject = *NewSequence;
+					}
+				}
+			}
+		}
+
+		for(UProperty* ComplexProp : ComplexAnimVariableProperties)
+		{
+			if(UObject** ResolvedObject = ComplexProp->ContainerPtrToValuePtr<UObject*>(DefaultObject))
+			{
+				if(UAnimationAsset* Sequence = Cast<UAnimationAsset>(*ResolvedObject))
+				{
+					// Found an anim sequence variable that's set to a valid sequence
+					if(UAnimationAsset* const* NewSequence = ComplexAnimMap.Find(Sequence))
+					{
+						*ResolvedObject = *NewSequence;
+					}
+				}
+			}
+		}
+	}
+
+	void GetBlueprintAssetVariableProperties(UAnimBlueprint* InBlueprint, TArray<UProperty*>& OutSimpleProperties, TArray<UProperty*>& OutComplexProperties)
+	{
+		OutSimpleProperties.Empty();
+		OutComplexProperties.Empty();
+
+		UAnimBlueprintGeneratedClass* GeneratedClass = InBlueprint->GetAnimBlueprintGeneratedClass();
+
+		for(FBPVariableDescription& VarDesc : InBlueprint->NewVariables)
+		{
+			// Grab any variables directly referencing and anim sequence object.
+			if(VarDesc.VarType.PinCategory == FString("object"))
+			{
+				if(UObject* VarSubObject = VarDesc.VarType.PinSubCategoryObject.Get())
+				{
+					// Get the object class
+					if(UClass* SubObjectAsClass = Cast<UClass>(VarSubObject))
+					{
+						// Anything that IS an anim sequence is a simple anim
+						// Anything that ISN'T an anim sequence but is an animation asset is a complex anim
+						if(SubObjectAsClass == UAnimSequence::StaticClass())
+						{
+							if(UProperty* Prop = GeneratedClass->FindPropertyByName(VarDesc.VarName))
+							{
+								OutSimpleProperties.Add(Prop);
+							}
+						}
+						else if(SubObjectAsClass->IsChildOf(UAnimationAsset::StaticClass()))
+						{
+							if(UProperty* Prop = GeneratedClass->FindPropertyByName(VarDesc.VarName))
+							{
+								OutComplexProperties.Add(Prop);
+							}
+						}
+					}
+				}
+			}
+		}
 	}
 
 	void CopyAnimCurves(USkeleton* OldSkeleton, USkeleton* NewSkeleton, UAnimSequenceBase *SequenceBase, const FName ContainerName, FRawCurveTracks::ESupportedCurveType CurveType )
diff --git a/Engine/Source/Editor/UnrealEd/Private/EditorEngine.cpp b/Engine/Source/Editor/UnrealEd/Private/EditorEngine.cpp
--- a/Engine/Source/Editor/UnrealEd/Private/EditorEngine.cpp
+++ b/Engine/Source/Editor/UnrealEd/Private/EditorEngine.cpp
@@ -3568,9 +3568,9 @@ void UEditorEngine::ParentActors( AActor* ParentActor, AActor* ChildActor, const
 		ParentActor->Modify();
 
 		// If child is already attached to something, modify the old parent and detach
-		if(ChildRoot->AttachParent != NULL)
+		if(ChildRoot->GetAttachParent() != nullptr)
 		{
-			AActor* OldParentActor = ChildRoot->AttachParent->GetOwner();
+			AActor* OldParentActor = ChildRoot->GetAttachParent()->GetOwner();
 			OldParentActor->Modify();
 			ChildRoot->DetachFromParent(true);
 
@@ -3580,7 +3580,7 @@ void UEditorEngine::ParentActors( AActor* ParentActor, AActor* ChildActor, const
 		// If the parent is already attached to this child, modify its parent and detach so we can allow the attachment
 		if(ParentRoot->IsAttachedTo(ChildRoot))
 		{
-			ParentRoot->AttachParent->GetOwner()->Modify();
+			ParentRoot->GetAttachParent()->GetOwner()->Modify();
 			ParentRoot->DetachFromParent(true);
 		}
 
@@ -3604,9 +3604,9 @@ bool UEditorEngine::DetachSelectedActors()
 		checkSlow( Actor );
 
 		USceneComponent* RootComp = Actor->GetRootComponent();
-		if( RootComp != NULL && RootComp->AttachParent != NULL)
+		if( RootComp != nullptr && RootComp->GetAttachParent() != nullptr)
 		{
-			AActor* OldParentActor = RootComp->AttachParent->GetOwner();
+			AActor* OldParentActor = RootComp->GetAttachParent()->GetOwner();
 			OldParentActor->Modify();
 			RootComp->DetachFromParent(true);
 			bDetachOccurred = true;
@@ -6149,7 +6149,7 @@ void UEditorEngine::UpdateAutoLoadProject()
 			{
 				FSuppressableWarningDialog::FSetupInfo Info( LOCTEXT("UpdateMacOSX_Body","Please update to the latest version of Mac OS X for best performance."), LOCTEXT("UpdateMacOSX_Title","Update Mac OS X"), TEXT("UpdateMacOSX"), GEditorSettingsIni );
 				Info.ConfirmText = LOCTEXT( "OK", "OK");
-				Info.bDefaultToSupressInTheFuture = true;
+				Info.bDefaultToSuppressInTheFuture = true;
 				FSuppressableWarningDialog OSUpdateWarning( Info );
 				OSUpdateWarning.ShowModal();
 			}
@@ -6196,7 +6196,7 @@ void UEditorEngine::UpdateAutoLoadProject()
 			{
 				FSuppressableWarningDialog::FSetupInfo Info( LOCTEXT("UnsupportedGPUWarning_Body","The current graphics card does not meet the minimum specification, for best performance an NVIDIA GeForce 470 GTX or AMD Radeon 6870 HD series card or higher is recommended. Rendering performance and compatibility are not guaranteed with this graphics card."), LOCTEXT("UnsupportedGPUWarning_Title","Unsupported Graphics Card"), TEXT("UnsupportedGPUWarning"), GEditorSettingsIni );
 				Info.ConfirmText = LOCTEXT( "OK", "OK");
-				Info.bDefaultToSupressInTheFuture = true;
+				Info.bDefaultToSuppressInTheFuture = true;
 				FSuppressableWarningDialog OSUpdateWarning( Info );
 				OSUpdateWarning.ShowModal();
 			}
@@ -6243,7 +6243,7 @@ void UEditorEngine::UpdateAutoLoadProject()
 			{
 				FSuppressableWarningDialog::FSetupInfo Info( LOCTEXT("SlowGPUWarning_Body","The current graphics card is slower than the recommanded specification of an NVIDIA GeForce 470 GTX or AMD Radeon 6870 HD series card or higher, performance may be low."), LOCTEXT("SlowGPUWarning_Title","Slow Graphics Card"), TEXT("SlowGPUWarning"), GEditorSettingsIni );
 				Info.ConfirmText = LOCTEXT( "OK", "OK");
-				Info.bDefaultToSupressInTheFuture = true;
+				Info.bDefaultToSuppressInTheFuture = true;
 				FSuppressableWarningDialog OSUpdateWarning( Info );
 				OSUpdateWarning.ShowModal();
 			}
@@ -6260,7 +6260,7 @@ void UEditorEngine::UpdateAutoLoadProject()
 			{
 				FSuppressableWarningDialog::FSetupInfo Info( LOCTEXT("LowRAMWarning_Body","For best performance install at least 8GB of RAM."), LOCTEXT("LowRAMWarning_Title","Low RAM"), TEXT("LowRAMWarning"), GEditorSettingsIni );
 				Info.ConfirmText = LOCTEXT( "OK", "OK");
-				Info.bDefaultToSupressInTheFuture = true;
+				Info.bDefaultToSuppressInTheFuture = true;
 				FSuppressableWarningDialog OSUpdateWarning( Info );
 				OSUpdateWarning.ShowModal();
 			}
@@ -6277,7 +6277,7 @@ void UEditorEngine::UpdateAutoLoadProject()
 			{
 				FSuppressableWarningDialog::FSetupInfo Info( LOCTEXT("SlowCPUWarning_Body","For best performance a Quad-core Intel or AMD processor, 2.5 GHz or faster is recommended."), LOCTEXT("SlowCPUWarning_Title","CPU Performance Warning"), TEXT("SlowCPUWarning"), GEditorSettingsIni );
 				Info.ConfirmText = LOCTEXT( "OK", "OK");
-				Info.bDefaultToSupressInTheFuture = true;
+				Info.bDefaultToSuppressInTheFuture = true;
 				FSuppressableWarningDialog OSUpdateWarning( Info );
 				OSUpdateWarning.ShowModal();
 			}
diff --git a/Engine/Source/Editor/UnrealEd/Private/EditorSelectUtils.cpp b/Engine/Source/Editor/UnrealEd/Private/EditorSelectUtils.cpp
--- a/Engine/Source/Editor/UnrealEd/Private/EditorSelectUtils.cpp
+++ b/Engine/Source/Editor/UnrealEd/Private/EditorSelectUtils.cpp
@@ -675,7 +675,7 @@ bool UUnrealEdEngine::IsComponentSelected(const UPrimitiveComponent* PrimCompone
 	bool bIsSelected = false;
 	if (GetSelectedComponentCount() > 0)
 	{
-		bIsSelected = GetSelectedComponents()->IsSelected(PrimComponent->IsEditorOnly() ? PrimComponent->AttachParent : PrimComponent);
+		bIsSelected = GetSelectedComponents()->IsSelected(PrimComponent->IsEditorOnly() ? PrimComponent->GetAttachParent() : PrimComponent);
 	}
 
 	return bIsSelected;
diff --git a/Engine/Source/Editor/UnrealEd/Private/EditorViewportCommands.cpp b/Engine/Source/Editor/UnrealEd/Private/EditorViewportCommands.cpp
--- a/Engine/Source/Editor/UnrealEd/Private/EditorViewportCommands.cpp
+++ b/Engine/Source/Editor/UnrealEd/Private/EditorViewportCommands.cpp
@@ -28,7 +28,11 @@ void FEditorViewportCommands::RegisterCommands()
 	UI_COMMAND( TexelFactorAccuracyMode, "Texel Factor Accuracy View Mode", "Visualize how accurate is the texel factor is computed on each mesh", EUserInterfaceActionType::RadioButton, FInputChord() );
 	UI_COMMAND( StationaryLightOverlapMode, "Stationary Light Overlap View Mode", "Visualizes overlap of stationary lights", EUserInterfaceActionType::RadioButton, FInputChord() );
 	UI_COMMAND( LightmapDensityMode, "Lightmap Density View Mode", "Renders the scene with lightmap density visualization", EUserInterfaceActionType::RadioButton, FInputChord( EModifierKey::Alt, EKeys::Zero ) );
+
+	UI_COMMAND( GroupLODColorationMode, "Level of Detail Coloration View Mode", "Renders the scene using Level of Detail visualization", EUserInterfaceActionType::RadioButton, FInputChord());
 	UI_COMMAND( LODColorationMode, "LOD Coloration View Mode", "Renders the scene using LOD color visualization", EUserInterfaceActionType::RadioButton, FInputChord() );
+	UI_COMMAND( HLODColorationMode, "HLOD Coloration View Mode", "Renders the scene using HLOD color visualization", EUserInterfaceActionType::RadioButton, FInputChord());
+
 	UI_COMMAND( VisualizeBufferMode, "Buffer Visualization View Mode", "Renders a set of selected post process materials, which visualize various intermediate render buffers (material attributes)", EUserInterfaceActionType::RadioButton, FInputChord());
 	UI_COMMAND( ReflectionOverrideMode, "Reflections View Mode", "Renders the scene with reflections only", EUserInterfaceActionType::RadioButton, FInputChord() );
 	UI_COMMAND( CollisionPawn, "Player Collision", "Renders player collision visualization", EUserInterfaceActionType::RadioButton, FInputChord() );
diff --git a/Engine/Source/Editor/UnrealEd/Private/Fbx/FbxAnimationExport.cpp b/Engine/Source/Editor/UnrealEd/Private/Fbx/FbxAnimationExport.cpp
--- a/Engine/Source/Editor/UnrealEd/Private/Fbx/FbxAnimationExport.cpp
+++ b/Engine/Source/Editor/UnrealEd/Private/Fbx/FbxAnimationExport.cpp
@@ -90,14 +90,21 @@ void FFbxExporter::ExportAnimSequenceToFbx(const UAnimSequence* AnimSeq,
 		FbxNode* CurrentBoneNode = BoneNodes[BoneIndex];
 
 		// Create the AnimCurves
-		FbxAnimCurve* Curves[6];
+		const uint32 NumberOfCurves = 9;
+		FbxAnimCurve* Curves[NumberOfCurves];
+		
+		// Individual curves for translation, rotation and scaling
 		Curves[0] = CurrentBoneNode->LclTranslation.GetCurve(InAnimLayer, FBXSDK_CURVENODE_COMPONENT_X, true);
 		Curves[1] = CurrentBoneNode->LclTranslation.GetCurve(InAnimLayer, FBXSDK_CURVENODE_COMPONENT_Y, true);
 		Curves[2] = CurrentBoneNode->LclTranslation.GetCurve(InAnimLayer, FBXSDK_CURVENODE_COMPONENT_Z, true);
-
+		
 		Curves[3] = CurrentBoneNode->LclRotation.GetCurve(InAnimLayer, FBXSDK_CURVENODE_COMPONENT_X, true);
 		Curves[4] = CurrentBoneNode->LclRotation.GetCurve(InAnimLayer, FBXSDK_CURVENODE_COMPONENT_Y, true);
 		Curves[5] = CurrentBoneNode->LclRotation.GetCurve(InAnimLayer, FBXSDK_CURVENODE_COMPONENT_Z, true);
+		
+		Curves[6] = CurrentBoneNode->LclScaling.GetCurve(InAnimLayer, FBXSDK_CURVENODE_COMPONENT_X, true);
+		Curves[7] = CurrentBoneNode->LclScaling.GetCurve(InAnimLayer, FBXSDK_CURVENODE_COMPONENT_Y, true);
+		Curves[8] = CurrentBoneNode->LclScaling.GetCurve(InAnimLayer, FBXSDK_CURVENODE_COMPONENT_Z, true);
 
 		float AnimTime					= AnimStartOffset;
 		float AnimEndTime				= (AnimSeq->SequenceLength - AnimEndOffset);
@@ -118,10 +125,10 @@ void FFbxExporter::ExportAnimSequenceToFbx(const UAnimSequence* AnimSeq,
 			// If this sequence does not have a track for the current bone, then skip it
 			continue;
 		}
-
-		for(int32 i = 0; i < 6; ++i)
+		
+		for (FbxAnimCurve* Curve : Curves)
 		{
-			Curves[i]->KeyModifyBegin();
+			Curve->KeyModifyBegin();
 		}
 
 		bool bLastKey = false;
@@ -132,40 +139,39 @@ void FFbxExporter::ExportAnimSequenceToFbx(const UAnimSequence* AnimSeq,
 			AnimSeq->GetBoneTransform(BoneAtom, BoneTrackIndex, AnimTime, true);
 
 			FbxVector4 Translation = Converter.ConvertToFbxPos(BoneAtom.GetTranslation());
-
 			FbxVector4 Rotation = Converter.ConvertToFbxRot(BoneAtom.GetRotation().Euler());
+			FbxVector4 Scale = Converter.ConvertToFbxScale(BoneAtom.GetScale3D());
+			FbxVector4 Vectors[3] = { Translation, Rotation, Scale };
 		
 			int32 lKeyIndex;
 
 			bLastKey = AnimTime >= AnimEndTime;
-			for(int32 i = 0, j=3; i < 3; ++i, ++j)
+			
+			// Loop over each curve and channel to set correct values
+			for (uint32 CurveIndex = 0; CurveIndex < 3; ++CurveIndex)
 			{
-				lKeyIndex = Curves[i]->KeyAdd(ExportTime);
-				Curves[i]->KeySetValue(lKeyIndex, Translation[i]);
-				Curves[i]->KeySetInterpolation(lKeyIndex, bLastKey ? FbxAnimCurveDef::eInterpolationConstant : FbxAnimCurveDef::eInterpolationCubic);
-
-				if( bLastKey )
+				for (uint32 ChannelIndex = 0; ChannelIndex < 3; ++ChannelIndex)
 				{
-					Curves[i]->KeySetConstantMode( lKeyIndex, FbxAnimCurveDef::eConstantStandard );
-				}
+					uint32 OffsetCurveIndex = (CurveIndex * 3) + ChannelIndex;
 
-				lKeyIndex = Curves[j]->KeyAdd(ExportTime);
-				Curves[j]->KeySetValue(lKeyIndex, Rotation[i]);
-				Curves[j]->KeySetInterpolation(lKeyIndex, bLastKey ? FbxAnimCurveDef::eInterpolationConstant : FbxAnimCurveDef::eInterpolationCubic);
+					lKeyIndex = Curves[OffsetCurveIndex]->KeyAdd(ExportTime);
+					Curves[OffsetCurveIndex]->KeySetValue(lKeyIndex, Vectors[CurveIndex][ChannelIndex]);
+					Curves[OffsetCurveIndex]->KeySetInterpolation(lKeyIndex, bLastKey ? FbxAnimCurveDef::eInterpolationConstant : FbxAnimCurveDef::eInterpolationCubic);
 
-				if( bLastKey )
-				{
-					Curves[j]->KeySetConstantMode( lKeyIndex, FbxAnimCurveDef::eConstantStandard );
-				}
+					if (bLastKey)
+					{
+						Curves[OffsetCurveIndex]->KeySetConstantMode(lKeyIndex, FbxAnimCurveDef::eConstantStandard);
+					}
+				}				
 			}
 
 			ExportTime += ExportTimeIncrement;
 			AnimTime += AnimTimeIncrement;
 		}
 
-		for(int32 i = 0; i < 6; ++i)
+		for (FbxAnimCurve* Curve : Curves)
 		{
-			Curves[i]->KeyModifyEnd();
+			Curve->KeyModifyEnd();
 		}
 	}
 }
diff --git a/Engine/Source/Editor/UnrealEd/Private/Fbx/FbxMainExport.cpp b/Engine/Source/Editor/UnrealEd/Private/Fbx/FbxMainExport.cpp
--- a/Engine/Source/Editor/UnrealEd/Private/Fbx/FbxMainExport.cpp
+++ b/Engine/Source/Editor/UnrealEd/Private/Fbx/FbxMainExport.cpp
@@ -281,7 +281,7 @@ static void SortActorsHierarchy(TArray<AActor*>& Actors)
 			Depth = 0;
 			if (InActor->GetRootComponent())
 			{
-				for (const USceneComponent* Test = InActor->GetRootComponent()->AttachParent; Test != nullptr; Test = Test->AttachParent, Depth++);
+				for (const USceneComponent* Test = InActor->GetRootComponent()->GetAttachParent(); Test != nullptr; Test = Test->GetAttachParent(), Depth++);
 			}
 		}
 		return Depth;
diff --git a/Engine/Source/Editor/UnrealEd/Private/HierarchicalLOD.cpp b/Engine/Source/Editor/UnrealEd/Private/HierarchicalLOD.cpp
--- a/Engine/Source/Editor/UnrealEd/Private/HierarchicalLOD.cpp
+++ b/Engine/Source/Editor/UnrealEd/Private/HierarchicalLOD.cpp
@@ -66,7 +66,7 @@ void FHierarchicalLODBuilder::Build()
 	// Fire map check warnings for hidden levels 
 	if (bVisibleLevelsWarning)
 	{
-		FMessageLog MapCheck("MapCheck");
+		FMessageLog MapCheck("HLODResults");
 		MapCheck.Warning()
 			->AddToken(FUObjectToken::Create(WorldSetting))
 			->AddToken(FTextToken::Create(LOCTEXT("MapCheck_Message_HLODHiddenLevels", "Certain levels are marked as hidden, Hierarchical LODs will not be build for hidden levels.")));
@@ -96,7 +96,7 @@ void FHierarchicalLODBuilder::PreviewBuild()
 	// Fire map check warnings for hidden levels 
 	if (bVisibleLevelsWarning)
 	{
-		FMessageLog MapCheck("MapCheck");
+		FMessageLog MapCheck("HLODResults");
 		MapCheck.Warning()
 			->AddToken(FUObjectToken::Create(WorldSetting))
 			->AddToken(FTextToken::Create(LOCTEXT("MapCheck_Message_PreviewBuild_HLODHiddenLevels", "Certain levels are marked as hidden, Hierarchical LODs will not be built for hidden levels.")));
@@ -174,7 +174,7 @@ void FHierarchicalLODBuilder::BuildClusters(ULevel* InLevel, const bool bCreateM
 	else
 	{
 		// Fire map check warnings if HLOD System is not enabled
-		FMessageLog MapCheck("MapCheck");
+		FMessageLog MapCheck("HLODResults");
 		MapCheck.Warning()
 			->AddToken(FUObjectToken::Create(InLevel->GetWorld()->GetWorldSettings()))
 			->AddToken(FTextToken::Create(LOCTEXT("MapCheck_Message_HLODSystemNotEnabled", "Hierarchical LOD System is disabled, unable to build LOD actors.")))
@@ -456,7 +456,7 @@ void FHierarchicalLODBuilder::BuildMeshesForLODActors()
 			continue;
 		}
 
-		FScopedSlowTask SlowTask(100, (LOCTEXT("HierarchicalLOD_BuildLODActorMeshes", "Building LODActor meshes")));
+		FScopedSlowTask SlowTask(105, (LOCTEXT("HierarchicalLOD_BuildLODActorMeshes", "Building LODActor meshes")));
 		SlowTask.MakeDialog();
 
 		TArray<TArray<ALODActor*>> LODLevelActors;
@@ -475,7 +475,7 @@ void FHierarchicalLODBuilder::BuildMeshesForLODActors()
 					if (Actor && Actor->IsA<ALODActor>())
 					{
 						ALODActor* LODActor = CastChecked<ALODActor>(Actor);
-
+						
 						if (LODActor->IsDirty() && LODActor->SubActors.Num() > 1)
 						{
 							LODLevelActors[LODActor->LODLevel - 1].Add(LODActor);
@@ -486,14 +486,17 @@ void FHierarchicalLODBuilder::BuildMeshesForLODActors()
 			
 				bool bBuildSuccesfull = true;
 				const int32 NumLODLevels = LODLevelActors.Num();
+				
 				for (int32 LODIndex = 0; LODIndex < NumLODLevels; ++LODIndex)
 				{
 					int32 CurrentLODLevel = LODIndex;
-					TArray<ALODActor*>& LODLevel = LODLevelActors[CurrentLODLevel];
+					int32 LODActorIndex = 0;
+					TArray<ALODActor*>& LODLevel = LODLevelActors[CurrentLODLevel];					
 					for (ALODActor* Actor : LODLevel)
 					{
 						bBuildSuccesfull &= FHierarchicalLODUtilities::BuildStaticMeshForLODActor(Actor, AssetsOuter, BuildLODLevelSettings[CurrentLODLevel], CurrentLODLevel);
-						SlowTask.EnterProgressFrame(100.0f / (float)NumLODActors);
+						SlowTask.EnterProgressFrame(100.0f / (float)NumLODActors, FText::Format(LOCTEXT("HierarchicalLOD_BuildLODActorMeshesProgress", "Building LODActor Mesh {1} / {2} in LOD Level {0}"), FText::AsNumber(LODIndex), FText::AsNumber(LODActorIndex), FText::AsNumber(LODLevelActors[CurrentLODLevel].Num())));
+						++LODActorIndex;
 					}
 				}
 				
@@ -552,8 +555,15 @@ void FHierarchicalLODBuilder::BuildMeshForLODActor(ALODActor* LODActor, const ui
 	BuildLODLevelSettings = WorldSetting->HierarchicalLODSetup;
 	UPackage* AssetsOuter = LODActor->GetLevel()->GetOutermost();
 
-	const bool bResult = FHierarchicalLODUtilities::BuildStaticMeshForLODActor(LODActor, AssetsOuter, BuildLODLevelSettings[LODLevel], LODLevel);	
-	check(bResult);
+	const bool bResult = FHierarchicalLODUtilities::BuildStaticMeshForLODActor(LODActor, AssetsOuter, BuildLODLevelSettings[LODLevel], LODLevel);
+
+	if (bResult == false)
+	{
+		FMessageLog("HLODResults").Error()
+			->AddToken(FTextToken::Create(LOCTEXT("HLODError_MeshNotBuildOne", "Cannot create proxy mesh for ")))
+			->AddToken(FUObjectToken::Create(LODActor))
+			->AddToken(FTextToken::Create(LOCTEXT("HLODError_MeshNotBuildTwo", " this could be caused by incorrect mesh components in the sub actors")));			
+	}
 }
 
 void FHierarchicalLODBuilder::MergeClustersAndBuildActors(ULevel* InLevel, const int32 LODIdx, float HighestCost, int32 MinNumActors, const bool bCreateMeshes)
diff --git a/Engine/Source/Editor/UnrealEd/Private/Kismet2/KismetReinstanceUtilities.cpp b/Engine/Source/Editor/UnrealEd/Private/Kismet2/KismetReinstanceUtilities.cpp
--- a/Engine/Source/Editor/UnrealEd/Private/Kismet2/KismetReinstanceUtilities.cpp
+++ b/Engine/Source/Editor/UnrealEd/Private/Kismet2/KismetReinstanceUtilities.cpp
@@ -950,9 +950,9 @@ void FActorReplacementHelper::CacheAttachInfo(const AActor* OldActor)
 
 	if (USceneComponent* OldRootComponent = OldActor->GetRootComponent())
 	{
-		if (OldRootComponent->AttachParent != nullptr)
+		if (OldRootComponent->GetAttachParent() != nullptr)
 		{
-			TargetAttachParent = OldRootComponent->AttachParent->GetOwner();
+			TargetAttachParent = OldRootComponent->GetAttachParent()->GetOwner();
 			// Root component should never be attached to another component in the same actor!
 			if (TargetAttachParent == OldActor)
 			{
@@ -960,8 +960,8 @@ void FActorReplacementHelper::CacheAttachInfo(const AActor* OldActor)
 				TargetAttachParent = nullptr;
 			}
 
-			TargetAttachSocket    = OldRootComponent->AttachSocketName;
-			TargetParentComponent = OldRootComponent->AttachParent;
+			TargetAttachSocket    = OldRootComponent->GetAttachSocketName();
+			TargetParentComponent = OldRootComponent->GetAttachParent();
 			
 			// detach it to remove any scaling
 			OldRootComponent->DetachFromParent(/*bMaintainWorldPosition =*/true);
@@ -982,12 +982,12 @@ void FActorReplacementHelper::CacheChildAttachments(const AActor* OldActor)
 	for (AActor* AttachedActor : AttachedActors)
 	{
 		USceneComponent* AttachedActorRoot = AttachedActor->GetRootComponent();
-		if (AttachedActorRoot && AttachedActorRoot->AttachParent)
+		if (AttachedActorRoot && AttachedActorRoot->GetAttachParent())
 		{
 			// Save info about actor to reattach
 			FAttachedActorInfo Info;
 			Info.AttachedActor = AttachedActor;
-			Info.AttachedToSocket = AttachedActorRoot->AttachSocketName;
+			Info.AttachedToSocket = AttachedActorRoot->GetAttachSocketName();
 			PendingChildAttachments.Add(Info);
 
 			// Now detach it
@@ -1035,7 +1035,7 @@ void FActorReplacementHelper::AttachChildActors(USceneComponent* RootComponent,
 		if (!Info.AttachedActor->IsPendingKill() && Info.AttachedActor->GetAttachParentActor() == nullptr)
 		{
 			USceneComponent* ChildRoot = Info.AttachedActor->GetRootComponent();
-			if (ChildRoot && ChildRoot->AttachParent != RootComponent)
+			if (ChildRoot && ChildRoot->GetAttachParent() != RootComponent)
 			{
 				ChildRoot->AttachTo(RootComponent, Info.AttachedToSocket, EAttachLocation::KeepWorldPosition);
 				ChildRoot->UpdateComponentToWorld();
diff --git a/Engine/Source/Editor/UnrealEd/Private/LevelViewportClickHandlers.cpp b/Engine/Source/Editor/UnrealEd/Private/LevelViewportClickHandlers.cpp
--- a/Engine/Source/Editor/UnrealEd/Private/LevelViewportClickHandlers.cpp
+++ b/Engine/Source/Editor/UnrealEd/Private/LevelViewportClickHandlers.cpp
@@ -281,7 +281,7 @@ namespace ClickHandlers
 		//If the component selected is editor-only, we want to select the non-editor-only component it's attached to
 		while (Component != nullptr && Component->IsEditorOnly())
 		{
-			Component = Component->AttachParent;
+			Component = Component->GetAttachParent();
 		}
 		
 		if (!ensure(Component != nullptr))
diff --git a/Engine/Source/Editor/UnrealEd/Private/PhysicsAssetUtils.cpp b/Engine/Source/Editor/UnrealEd/Private/PhysicsAssetUtils.cpp
--- a/Engine/Source/Editor/UnrealEd/Private/PhysicsAssetUtils.cpp
+++ b/Engine/Source/Editor/UnrealEd/Private/PhysicsAssetUtils.cpp
@@ -195,8 +195,8 @@ bool CreateFromSkeletalMeshInternal(UPhysicsAsset* PhysicsAsset, USkeletalMesh*
 
 				CS->DefaultInstance.ConstraintBone2 = ParentName;
 				CS->DefaultInstance.Pos2 = RelTM.GetOrigin();
-				CS->DefaultInstance.PriAxis2 = RelTM.GetScaledAxis(EAxis::X);
-				CS->DefaultInstance.SecAxis2 = RelTM.GetScaledAxis(EAxis::Y);
+				CS->DefaultInstance.PriAxis2 = RelTM.GetUnitAxis(EAxis::X);
+				CS->DefaultInstance.SecAxis2 = RelTM.GetUnitAxis(EAxis::Y);
 
 				// Disable collision between constrained bodies by default.
 				PhysicsAsset->DisableCollision(NewBodyIndex, ParentBodyIndex);
diff --git a/Engine/Source/Editor/UnrealEd/Private/SEditorViewport.cpp b/Engine/Source/Editor/UnrealEd/Private/SEditorViewport.cpp
--- a/Engine/Source/Editor/UnrealEd/Private/SEditorViewport.cpp
+++ b/Engine/Source/Editor/UnrealEd/Private/SEditorViewport.cpp
@@ -355,7 +355,9 @@ void SEditorViewport::BindCommands()
 	MAP_VIEWMODE_ACTION( Commands.StationaryLightOverlapMode, VMI_StationaryLightOverlap );
 	MAP_VIEWMODE_ACTION( Commands.LightmapDensityMode, VMI_LightmapDensity );
 	MAP_VIEWMODE_ACTION( Commands.ReflectionOverrideMode, VMI_ReflectionOverride );
+	MAP_VIEWMODE_ACTION( Commands.GroupLODColorationMode, VMI_GroupLODColoration);
 	MAP_VIEWMODE_ACTION( Commands.LODColorationMode, VMI_LODColoration );
+	MAP_VIEWMODE_ACTION( Commands.HLODColorationMode, VMI_HLODColoration);
 	MAP_VIEWMODE_ACTION( Commands.VisualizeBufferMode, VMI_VisualizeBuffer );
 	MAP_VIEWMODE_ACTION( Commands.CollisionPawn, VMI_CollisionPawn);
 	MAP_VIEWMODE_ACTION( Commands.CollisionVisibility, VMI_CollisionVisibility);
diff --git a/Engine/Source/Editor/UnrealEd/Private/SEditorViewportViewMenu.cpp b/Engine/Source/Editor/UnrealEd/Private/SEditorViewportViewMenu.cpp
--- a/Engine/Source/Editor/UnrealEd/Private/SEditorViewportViewMenu.cpp
+++ b/Engine/Source/Editor/UnrealEd/Private/SEditorViewportViewMenu.cpp
@@ -110,6 +110,9 @@ FText SEditorViewportViewMenu::GetViewMenuLabel() const
 			case VMI_LODColoration:
 				Label = LOCTEXT("ViewMenuTitle_LODColoration", "LOD Coloration");
 				break;
+			case VMI_HLODColoration:
+				Label = LOCTEXT("ViewMenuTitle_HLODColoration", "HLOD Coloration");
+				break;
 		}
 	}
 
@@ -216,6 +219,14 @@ const FSlateBrush* SEditorViewportViewMenu::GetViewMenuLabelIcon() const
 			case VMI_LODColoration:
 				Icon = LODColorationIcon;
 				break;
+
+			case VMI_HLODColoration:
+				Icon = LODColorationIcon;
+				break;
+
+			case VMI_GroupLODColoration:
+				Icon = LODColorationIcon;
+				break;
 		}
 	}
 
diff --git a/Engine/Source/Editor/UnrealEd/Private/SSocketManager.cpp b/Engine/Source/Editor/UnrealEd/Private/SSocketManager.cpp
--- a/Engine/Source/Editor/UnrealEd/Private/SSocketManager.cpp
+++ b/Engine/Source/Editor/UnrealEd/Private/SSocketManager.cpp
@@ -699,9 +699,9 @@ void SSocketManager::OnSocketPropertyChanged( const UStaticMeshSocket* Socket, c
 						const USceneComponent* Root = Actor->GetRootComponent();
 						if (Root != nullptr)
 						{
-							for (USceneComponent* Child : Root->AttachChildren)
+							for (USceneComponent* Child : Root->GetAttachChildren())
 							{
-								if (Child != nullptr && Child->AttachSocketName == Socket->SocketName)
+								if (Child != nullptr && Child->GetAttachSocketName() == Socket->SocketName)
 								{
 									Child->UpdateComponentToWorld();
 									bUpdatedChild = true;
diff --git a/Engine/Source/Editor/UnrealEd/Private/VertexSnapping.cpp b/Engine/Source/Editor/UnrealEd/Private/VertexSnapping.cpp
--- a/Engine/Source/Editor/UnrealEd/Private/VertexSnapping.cpp
+++ b/Engine/Source/Editor/UnrealEd/Private/VertexSnapping.cpp
@@ -670,11 +670,10 @@ static void GetActorsToIgnore( AActor* Actor, TSet< TWeakObjectPtr<AActor> >& Ac
 
 		// We cannot snap to any attached children or actors in the same group as moving this actor will also move the children as we are snapping to them, 
 		// causing a cascading effect and unexpected results
-		TArray<USceneComponent*>& AttachedChildren = Actor->GetRootComponent()->AttachChildren;
+		const TArray<USceneComponent*>& AttachedChildren = Actor->GetRootComponent()->GetAttachChildren();
 
-		for( int32 ChildIndex = 0; ChildIndex < AttachedChildren.Num(); ++ChildIndex )
+		for (USceneComponent* Child : AttachedChildren)
 		{
-			USceneComponent* Child = AttachedChildren[ChildIndex];
 			if( Child && Child->GetOwner() )
 			{
 				ActorsToIgnore.Add( Child->GetOwner() );
diff --git a/Engine/Source/Editor/UnrealEd/Public/Dialogs/Dialogs.h b/Engine/Source/Editor/UnrealEd/Public/Dialogs/Dialogs.h
--- a/Engine/Source/Editor/UnrealEd/Public/Dialogs/Dialogs.h
+++ b/Engine/Source/Editor/UnrealEd/Public/Dialogs/Dialogs.h
@@ -78,7 +78,7 @@ public:
 		FString IniSettingFileName;
 
 		/** If true the suppress checkbox defaults to true*/
-		bool bDefaultToSupressInTheFuture;
+		bool bDefaultToSuppressInTheFuture;
 
 		/** Text used on the button which will return FSuppressableWarningDialog::Confirm */
 		FText ConfirmText;
@@ -105,7 +105,7 @@ public:
 			, Title(InTitle)
 			, IniSettingName(InIniSettingName)
 			, IniSettingFileName(InIniSettingFileName)
-			, bDefaultToSupressInTheFuture(false)
+			, bDefaultToSuppressInTheFuture(false)
 			, ConfirmText()
 			, CancelText()
 			, CheckBoxText(NSLOCTEXT("ModalDialogs", "DefaultCheckBoxMessage", "Don't show this again"))
diff --git a/Engine/Source/Editor/UnrealEd/Public/EditorAnimUtils.h b/Engine/Source/Editor/UnrealEd/Public/EditorAnimUtils.h
--- a/Engine/Source/Editor/UnrealEd/Public/EditorAnimUtils.h
+++ b/Engine/Source/Editor/UnrealEd/Public/EditorAnimUtils.h
@@ -156,6 +156,59 @@ namespace EditorAnimUtils
 		return ReturnMap;
 	}
 
+	void GetBlueprintAssetVariableProperties(UAnimBlueprint* InBlueprint, TArray<UProperty*>& OutSimpleProperties, TArray<UProperty*>& OutComplexProperties);
+
+	template<class AssetType>
+	void GetAssetsFromProperties(TArray<UProperty*> InProperties, UObject* Scope, TArray<AssetType*>& OutAssets)
+	{
+		check(Scope);
+
+		OutAssets.Empty();
+		for(UProperty* Prop : InProperties)
+		{
+			if(Prop)
+			{
+				if(UArrayProperty* ArrayProp = Cast<UArrayProperty>(Prop))
+				{
+					// Blueprint array
+					FScriptArrayHelper Helper(ArrayProp, Prop->ContainerPtrToValuePtr<uint8>(Scope));
+					const int32 ArrayNum = Helper.Num();
+					for(int32 Idx = 0; Idx < ArrayNum; ++Idx)
+					{
+						// These were gathered from UObject types so we know this should succeed
+						UObject** Object = (UObject**)Helper.GetRawPtr(Idx);
+						if(AssetType* Asset = Cast<AssetType>(*Object))
+						{
+							OutAssets.Add(Asset);
+						}
+					}
+				}
+				else if(Prop->ArrayDim > 1)
+				{
+					// Native array
+					for(int32 Idx = 0; Idx < Prop->ArrayDim; ++Idx)
+					{
+						if(UObject** ResolvedObject = Prop->ContainerPtrToValuePtr<UObject*>(Scope, Idx))
+						{
+							if(AssetType* Asset = Cast<AssetType>(*ResolvedObject))
+							{
+								OutAssets.Add(Asset);
+							}
+						}
+					}
+				}
+				else if(UObject** ResolvedObject = Prop->ContainerPtrToValuePtr<UObject*>(Scope))
+				{
+					// Normal property
+					if(AssetType* Asset = Cast<AssetType>(*ResolvedObject))
+					{
+						OutAssets.Add(Asset);
+					}
+				}
+			}
+		}
+	}
+
 	// utility functions
 	UNREALED_API void CopyAnimCurves(USkeleton* OldSkeleton, USkeleton* NewSkeleton, UAnimSequenceBase *SequenceBase, const FName ContainerName, FRawCurveTracks::ESupportedCurveType CurveType );
 } // namespace EditorAnimUtils
diff --git a/Engine/Source/Editor/UnrealEd/Public/EditorViewportCommands.h b/Engine/Source/Editor/UnrealEd/Public/EditorViewportCommands.h
--- a/Engine/Source/Editor/UnrealEd/Public/EditorViewportCommands.h
+++ b/Engine/Source/Editor/UnrealEd/Public/EditorViewportCommands.h
@@ -83,9 +83,15 @@ public:
 	/** Changes the viewport to vertex densities mode */
 	TSharedPtr< FUICommandInfo > VertexDensitiesMode;
 
+	/** Shows the group menu item for LOD and HLOD coloration modes */
+	TSharedPtr< FUICommandInfo > GroupLODColorationMode;
+
 	/** Changes the viewport to LOD index coloration mode */
 	TSharedPtr< FUICommandInfo > LODColorationMode;
 
+	/** Changes the viewport to HLOD index coloration mode */
+	TSharedPtr< FUICommandInfo > HLODColorationMode;
+
 	/** Changes the viewport to visualize the buffer content */
 	TSharedPtr< FUICommandInfo > VisualizeBufferMode;
 
diff --git a/Engine/Source/Programs/UnrealLightmass/Private/LightmassCore/Math/LMMath.cpp b/Engine/Source/Programs/UnrealLightmass/Private/LightmassCore/Math/LMMath.cpp
--- a/Engine/Source/Programs/UnrealLightmass/Private/LightmassCore/Math/LMMath.cpp
+++ b/Engine/Source/Programs/UnrealLightmass/Private/LightmassCore/Math/LMMath.cpp
@@ -48,9 +48,9 @@ FLinearColor FLinearColorUtils::LinearRGBToHSV(const FLinearColor& InColor)
 	const float RGBRange = RGBMax - RGBMin;
 
 	const float Hue = (RGBMax == RGBMin ? 0.0f :
-					   RGBMax == InColor.R    ? fmod((((InColor.G - InColor.B) / RGBRange) * 60.0f) + 360.0f, 360.0f) :
-					   RGBMax == InColor.G    ?      (((InColor.B - InColor.R) / RGBRange) * 60.0f) + 120.0f :
-					   RGBMax == InColor.B    ?      (((InColor.R - InColor.G) / RGBRange) * 60.0f) + 240.0f :
+					   RGBMax == InColor.R    ? FMath::Fmod((((InColor.G - InColor.B) / RGBRange) * 60.0f) + 360.0f, 360.0f) :
+					   RGBMax == InColor.G    ?             (((InColor.B - InColor.R) / RGBRange) * 60.0f) + 120.0f :
+					   RGBMax == InColor.B    ?             (((InColor.R - InColor.G) / RGBRange) * 60.0f) + 240.0f :
 					   0.0f);
 	
 	const float Saturation = (RGBMax == 0.0f ? 0.0f : RGBRange / RGBMax);
diff --git a/Engine/Source/Runtime/AnimGraphRuntime/Private/BoneControllers/AnimNode_CopyBoneDelta.cpp b/Engine/Source/Runtime/AnimGraphRuntime/Private/BoneControllers/AnimNode_CopyBoneDelta.cpp
new file mode 100644
--- /dev/null
+++ b/Engine/Source/Runtime/AnimGraphRuntime/Private/BoneControllers/AnimNode_CopyBoneDelta.cpp
@@ -0,0 +1,107 @@
+// Copyright 1998-2016 Epic Games, Inc. All Rights Reserved.
+
+#include "AnimGraphRuntimePrivatePCH.h"
+#include "AnimNode_CopyBoneDelta.h"
+
+FAnimNode_CopyBoneDelta::FAnimNode_CopyBoneDelta()
+	: bCopyTranslation(false)
+	, bCopyRotation(false)
+	, bCopyScale(false)
+	, CopyMode(CopyBoneDeltaMode::Accumulate)
+	, TranslationMultiplier(1.0f)
+	, RotationMultiplier(1.0f)
+	, ScaleMultiplier(1.0f)
+{
+
+}
+
+void FAnimNode_CopyBoneDelta::GatherDebugData(FNodeDebugData& DebugData)
+{
+	FString DebugLine = DebugData.GetNodeName(this);
+
+	DebugData.AddDebugItem(DebugLine);
+
+	ComponentPose.GatherDebugData(DebugData);
+}
+
+void FAnimNode_CopyBoneDelta::EvaluateBoneTransforms(USkeletalMeshComponent* SkelComp, FCSPose<FCompactPose>& MeshBases, TArray<FBoneTransform>& OutBoneTransforms)
+{
+	if(!bCopyTranslation && !bCopyRotation && !bCopyScale)
+	{
+		return;
+	}
+
+	const FBoneContainer& BoneContainer = MeshBases.GetPose().GetBoneContainer();
+	FCompactPoseBoneIndex SourceBoneIndex = SourceBone.GetCompactPoseIndex(BoneContainer);
+	FCompactPoseBoneIndex TargetBoneIndex = TargetBone.GetCompactPoseIndex(BoneContainer);
+
+	FTransform SourceTM = MeshBases.GetComponentSpaceTransform(SourceBoneIndex);
+	FTransform TargetTM = MeshBases.GetComponentSpaceTransform(TargetBoneIndex);
+
+	// Convert to parent space
+	FAnimationRuntime::ConvertCSTransformToBoneSpace(SkelComp, MeshBases, SourceTM, SourceBoneIndex, BCS_ParentBoneSpace);
+	FAnimationRuntime::ConvertCSTransformToBoneSpace(SkelComp, MeshBases, TargetTM, TargetBoneIndex, BCS_ParentBoneSpace);
+
+	// Ref pose transform
+	FTransform RefLSTransform = SkelComp->SkeletalMesh->RefSkeleton.GetRefBonePose()[SourceBone.GetMeshPoseIndex().GetInt()];
+
+	// Get transform relative to ref pose
+	SourceTM.SetToRelativeTransform(RefLSTransform);
+
+	if(CopyMode == CopyBoneDeltaMode::Accumulate)
+	{
+		if(bCopyTranslation)
+		{
+			TargetTM.AddToTranslation(SourceTM.GetTranslation() * TranslationMultiplier);
+		}
+		if(bCopyRotation)
+		{
+			FVector Axis;
+			float Angle;
+			SourceTM.GetRotation().ToAxisAndAngle(Axis, Angle);
+
+			TargetTM.SetRotation(FQuat(Axis, Angle * RotationMultiplier) * TargetTM.GetRotation());
+		}
+		if(bCopyScale)
+		{
+			TargetTM.SetScale3D(TargetTM.GetScale3D() * (SourceTM.GetScale3D() * ScaleMultiplier));
+		}
+	}
+	else //CopyMode = CopyBoneDeltaMode::Copy
+	{
+		if(bCopyTranslation)
+		{
+			TargetTM.SetTranslation(SourceTM.GetTranslation() * TranslationMultiplier);
+		}
+
+		if(bCopyRotation)
+		{
+			FVector Axis;
+			float Angle;
+			SourceTM.GetRotation().ToAxisAndAngle(Axis, Angle);
+
+			TargetTM.SetRotation(FQuat(Axis, Angle * RotationMultiplier));
+		}
+
+		if(bCopyScale)
+		{
+			TargetTM.SetScale3D(SourceTM.GetScale3D() * ScaleMultiplier);
+		}
+	}
+
+	// Back out to component space
+	FAnimationRuntime::ConvertBoneSpaceTransformToCS(SkelComp, MeshBases, TargetTM, TargetBoneIndex, BCS_ParentBoneSpace);
+
+	OutBoneTransforms.Add(FBoneTransform(TargetBoneIndex, TargetTM));
+}
+
+bool FAnimNode_CopyBoneDelta::IsValidToEvaluate(const USkeleton* Skeleton, const FBoneContainer& RequiredBones)
+{
+	return (TargetBone.IsValid(RequiredBones) && (TargetBone == SourceBone || SourceBone.IsValid(RequiredBones)));
+}
+
+void FAnimNode_CopyBoneDelta::InitializeBoneReferences(const FBoneContainer& RequiredBones)
+{
+	SourceBone.Initialize(RequiredBones);
+	TargetBone.Initialize(RequiredBones);
+}
diff --git a/Engine/Source/Runtime/AnimGraphRuntime/Private/BoneControllers/AnimNode_SpringBone.cpp b/Engine/Source/Runtime/AnimGraphRuntime/Private/BoneControllers/AnimNode_SpringBone.cpp
--- a/Engine/Source/Runtime/AnimGraphRuntime/Private/BoneControllers/AnimNode_SpringBone.cpp
+++ b/Engine/Source/Runtime/AnimGraphRuntime/Private/BoneControllers/AnimNode_SpringBone.cpp
@@ -97,9 +97,9 @@ void FAnimNode_SpringBone::EvaluateBoneTransforms(USkeletalMeshComponent* SkelCo
 	FVector const TargetPos = BoneTransformInWorldSpace.GetLocation();
 
 	AActor* SkelOwner = (SkelComp != NULL) ? SkelComp->GetOwner() : NULL;
-	if ((SkelComp != NULL) && (SkelComp->AttachParent != NULL) && (SkelOwner == NULL))
+	if ((SkelComp != NULL) && (SkelComp->GetAttachParent() != NULL) && (SkelOwner == NULL))
 	{
-		SkelOwner = SkelComp->AttachParent->GetOwner();
+		SkelOwner = SkelComp->GetAttachParent()->GetOwner();
 	}
 
 	// Init values first time
