Fix crash when deleting the last actor from an HLOD cluster.
Destroy cluster if empty (deleted actor can be removed from the LODActor list from EditorActor.cpp, leaving empty clusters around)
#rnx

diff --git a/Engine/Source/Developer/HierarchicalLODUtilities/Private/HierarchicalLODUtilities.cpp b/Engine/Source/Developer/HierarchicalLODUtilities/Private/HierarchicalLODUtilities.cpp
--- a/Engine/Source/Developer/HierarchicalLODUtilities/Private/HierarchicalLODUtilities.cpp
+++ b/Engine/Source/Developer/HierarchicalLODUtilities/Private/HierarchicalLODUtilities.cpp
@@ -931,7 +931,6 @@ void FHierarchicalLODUtilities::DestroyLODActor(ALODActor* InActor)
 	ALODActor* ParentActor = GetParentLODActor(InActor);
 
 	DestroyCluster(InActor);
-	World->DestroyActor(InActor);
 
 	if (ParentActor && !ParentActor->HasAnySubActors())
 	{
@@ -968,7 +967,6 @@ void FHierarchicalLODUtilities::DeleteLODActorsInHLODLevel(UWorld* InWorld, cons
 		if (LodActor && LodActor->LODLevel == (HLODLevelIndex + 1))
 		{
 			DestroyCluster(LodActor);
-			InWorld->DestroyActor(LodActor);
 		}
 	}
 }
diff --git a/Engine/Source/Editor/HierarchicalLODOutliner/Private/HLODOutliner.cpp b/Engine/Source/Editor/HierarchicalLODOutliner/Private/HLODOutliner.cpp
--- a/Engine/Source/Editor/HierarchicalLODOutliner/Private/HLODOutliner.cpp
+++ b/Engine/Source/Editor/HierarchicalLODOutliner/Private/HLODOutliner.cpp
@@ -1660,15 +1660,14 @@ namespace HLODOutliner
 				for (TWeakObjectPtr<ALODActor> Actor : ActorArray)
 				{
 					// Check if actor is not null due to Destroy Actor
-					if (Actor.IsValid())
+					if (Actor.IsValid() && Actor != InActor)
 					{
 						Actor->CleanSubActorArray();
-						const bool bRemovedSubActor = Actor->RemoveSubActor(InActor);
+						Actor->RemoveSubActor(InActor);
 
-						if (bRemovedSubActor && Actor->SubActors.Num() == 0)
+						if (Actor->SubActors.Num() == 0)
 						{
 							HierarchicalLODUtilities->DestroyCluster(Actor.Get());
-							Actor->GetWorld()->DestroyActor(Actor.Get());
 						}
 					}
 				}
