UE-69860
No progress bar is shown when building HLOD clusters from the outliner

diff --git a/Engine/Source/Editor/HierarchicalLODOutliner/Private/HLODOutliner.cpp b/Engine/Source/Editor/HierarchicalLODOutliner/Private/HLODOutliner.cpp
--- a/Engine/Source/Editor/HierarchicalLODOutliner/Private/HLODOutliner.cpp
+++ b/Engine/Source/Editor/HierarchicalLODOutliner/Private/HLODOutliner.cpp
@@ -49,6 +49,7 @@
 #include "Widgets/Notifications/SNotificationList.h"
 #include "Toolkits/AssetEditorManager.h"
 #include "Engine/HLODProxy.h"
+#include "Misc/ScopedSlowTask.h"
 
 #define LOCTEXT_NAMESPACE "HLODOutliner"
 
@@ -1098,6 +1099,9 @@ namespace HLODOutliner
 			// This call came from a context menu
 			auto SelectedItems = TreeView->GetSelectedItems();
 			
+			FScopedSlowTask SlowTask(SelectedItems.Num(), FText::Format(LOCTEXT("SHLODOutliner_BuildProxy", "Building Proxy {0}|plural(one=Mesh,other=Meshes)"), SelectedItems.Num()));
+			SlowTask.MakeDialog();
+
 			// Loop over all selected items (context menu can't be called with multiple items selected that aren't of the same types)
 			for (auto SelectedItem : SelectedItems )
 			{
@@ -1116,6 +1120,8 @@ namespace HLODOutliner
 						}
 					}
 				}
+
+				SlowTask.EnterProgressFrame();
 			}
 			
 			SetForcedLODLevel(ForcedLODLevel);
@@ -1135,6 +1141,9 @@ namespace HLODOutliner
 			// This call came from a context menu
 			auto SelectedItems = TreeView->GetSelectedItems();
 
+			FScopedSlowTask SlowTask(SelectedItems.Num(), FText::Format(LOCTEXT("SHLODOutliner_RebuildProxy", "Rebuild Proxy {0}|plural(one=Mesh,other=Meshes)"), SelectedItems.Num()));
+			SlowTask.MakeDialog();
+
 			// Loop over all selected items (context menu can't be called with multiple items selected that aren't of the same types)
 			for (auto SelectedItem : SelectedItems)
 			{
@@ -1150,6 +1159,8 @@ namespace HLODOutliner
 						CurrentWorld->HierarchicalLODBuilder->BuildMeshForLODActor(ActorItem->LODActor.Get(), LevelItem->LODLevelIndex);
 					}
 				}
+
+				SlowTask.EnterProgressFrame();
 			}
 
 			SetForcedLODLevel(ForcedLODLevel);
