HLOD Outliner - Add choice for "No HLOD" in the HLOD dropdown

diff --git a/Engine/Source/Editor/HierarchicalLODOutliner/Private/HLODOutliner.cpp b/Engine/Source/Editor/HierarchicalLODOutliner/Private/HLODOutliner.cpp
--- a/Engine/Source/Editor/HierarchicalLODOutliner/Private/HLODOutliner.cpp
+++ b/Engine/Source/Editor/HierarchicalLODOutliner/Private/HLODOutliner.cpp
@@ -57,12 +57,15 @@
 
 namespace HLODOutliner
 {
+	const int32 LOD_AUTO = -1;
+	const int32 LOD_DISABLED = -2;
+
 	SHLODOutliner::SHLODOutliner()		
 	{
 		bNeedsRefresh = true;
 		CurrentWorld = nullptr;
 		CurrentWorldSettings = nullptr;
-		ForcedLODLevel = -1;
+		ForcedLODLevel = LOD_AUTO;
 		bArrangeHorizontally = false;
 
 		FHierarchicalLODUtilitiesModule& Module = FModuleManager::LoadModuleChecked<FHierarchicalLODUtilitiesModule>("HierarchicalLODUtilities");
@@ -971,7 +974,9 @@ namespace HLODOutliner
 
 	FText SHLODOutliner::HandleForceLevelText() const
 	{
-		return ForcedLODLevel == -1 ? LOCTEXT("AutoLOD", "LOD Auto") : FText::Format(LOCTEXT("LODLevelFormat", "LOD {0}"), FText::AsNumber(ForcedLODLevel));
+		return ForcedLODLevel == LOD_AUTO ? LOCTEXT("AutoLOD", "LOD Auto") : 
+			   ForcedLODLevel == LOD_DISABLED ? LOCTEXT("DisabledLOD", "LOD Disabled") : 
+			   FText::Format(LOCTEXT("LODLevelFormat", "LOD {0}"), FText::AsNumber(ForcedLODLevel));
 	}
 
 	TSharedRef<SWidget> SHLODOutliner::GetForceLevelMenuContent() const
@@ -984,9 +989,21 @@ namespace HLODOutliner
 			LOCTEXT("AutoLODTooltip", "Determine LOD level automatically"),
 			FSlateIcon(),
 			FUIAction(
-				FExecuteAction::CreateSP(const_cast<SHLODOutliner*>(this), &SHLODOutliner::SetForcedLODLevel, -1), 
+				FExecuteAction::CreateSP(const_cast<SHLODOutliner*>(this), &SHLODOutliner::SetForcedLODLevel, LOD_AUTO), 
 				FCanExecuteAction(), 
-				FGetActionCheckState::CreateLambda([this](){ return ForcedLODLevel == -1 ? ECheckBoxState::Checked : ECheckBoxState::Unchecked; })),
+				FGetActionCheckState::CreateLambda([this](){ return ForcedLODLevel == LOD_AUTO ? ECheckBoxState::Checked : ECheckBoxState::Unchecked; })),
+			NAME_None,
+			EUserInterfaceActionType::RadioButton);
+
+		// Auto LOD
+		MenuBuilder.AddMenuEntry(
+			LOCTEXT("DisabledLOD", "LOD Disabled"),
+			LOCTEXT("DisabledLODTooltip", "Disable LOD and show original scene"),
+			FSlateIcon(),
+			FUIAction(
+				FExecuteAction::CreateSP(const_cast<SHLODOutliner*>(this), &SHLODOutliner::SetForcedLODLevel, LOD_DISABLED),
+				FCanExecuteAction(),
+				FGetActionCheckState::CreateLambda([this]() { return ForcedLODLevel == LOD_DISABLED ? ECheckBoxState::Checked : ECheckBoxState::Unchecked; })),
 			NAME_None,
 			EUserInterfaceActionType::RadioButton);
 
@@ -1017,7 +1034,7 @@ namespace HLODOutliner
 
 	void SHLODOutliner::RestoreForcedLODLevel(int32 LODLevel)
 	{
-		if (LODLevel == -1)
+		if (LODLevel == LOD_AUTO)
 		{
 			return;
 		}
@@ -1045,7 +1062,7 @@ namespace HLODOutliner
 	{
 		RestoreForcedLODLevel(ForcedLODLevel);
 
-		if (LODLevel == -1)
+		if (LODLevel == LOD_AUTO)
 		{
 			ForcedLODLevel = LODLevel;
 			return;
@@ -1075,7 +1092,7 @@ namespace HLODOutliner
 	void SHLODOutliner::ResetLODLevelForcing()
 	{
 		RestoreForcedLODLevel(ForcedLODLevel);
-		SetForcedLODLevel(-1);
+		SetForcedLODLevel(LOD_AUTO);
 	}
 
 	void SHLODOutliner::CreateHierarchicalVolumeForActor()
