HLOD Outliner - Stop updating while PIE is running

diff --git a/Engine/Source/Editor/HierarchicalLODOutliner/Private/HLODOutliner.cpp b/Engine/Source/Editor/HierarchicalLODOutliner/Private/HLODOutliner.cpp
--- a/Engine/Source/Editor/HierarchicalLODOutliner/Private/HLODOutliner.cpp
+++ b/Engine/Source/Editor/HierarchicalLODOutliner/Private/HLODOutliner.cpp
@@ -74,7 +74,10 @@ namespace HLODOutliner
 
 	SHLODOutliner::~SHLODOutliner()
 	{
-		DeregisterDelegates();	
+		DeregisterDelegates();
+		FEditorDelegates::BeginPIE.RemoveAll(this);
+		FEditorDelegates::EndPIE.RemoveAll(this);
+
 		DestroySelectionActors();
 		CurrentWorld = nullptr;
 		HLODTreeRoot.Empty();
@@ -299,7 +302,13 @@ namespace HLODOutliner
 				]
 			];
 
-		RegisterDelegates();
+		if (!CurrentWorld.IsValid() || !CurrentWorld->IsPlayInEditor())
+		{
+			RegisterDelegates();
+		}
+
+		FEditorDelegates::BeginPIE.AddRaw(this, &SHLODOutliner::OnBeginPieEvent);
+		FEditorDelegates::EndPIE.AddRaw(this, &SHLODOutliner::OnEndPieEvent);
 	}
 
 	TSharedRef<SWidget> SHLODOutliner::CreateMainButtonWidgets()
@@ -877,19 +886,32 @@ namespace HLODOutliner
 
 	END_SLATE_FUNCTION_BUILD_OPTIMIZATION
 
+	void SHLODOutliner::OnBeginPieEvent(bool bIsSimulating)
+	{
+		DeregisterDelegates();
+		FullRefresh();
+	}
+
+	void SHLODOutliner::OnEndPieEvent(bool bIsSimulating)
+	{
+		RegisterDelegates();
+		FullRefresh();
+	}
+
 	void SHLODOutliner::RegisterDelegates()
 	{
 		FEditorDelegates::MapChange.AddSP(this, &SHLODOutliner::OnMapChange);
 		FEditorDelegates::NewCurrentLevel.AddSP(this, &SHLODOutliner::OnNewCurrentLevel);
 		FEditorDelegates::OnMapOpened.AddSP(this, &SHLODOutliner::OnMapLoaded);
+
 		FWorldDelegates::LevelAddedToWorld.AddSP(this, &SHLODOutliner::OnLevelAdded);
 		FWorldDelegates::LevelRemovedFromWorld.AddSP(this, &SHLODOutliner::OnLevelRemoved);
+
 		GEngine->OnLevelActorListChanged().AddSP(this, &SHLODOutliner::FullRefresh);
 		GEngine->OnLevelActorAdded().AddSP(this, &SHLODOutliner::OnLevelActorsAdded);
 		GEngine->OnLevelActorDeleted().AddSP(this, &SHLODOutliner::OnLevelActorsRemoved);
 		GEngine->OnActorMoved().AddSP(this, &SHLODOutliner::OnActorMovedEvent);
 
-
 		// Selection change
 		USelection::SelectionChangedEvent.AddRaw(this, &SHLODOutliner::OnLevelSelectionChanged);
 		USelection::SelectObjectEvent.AddRaw(this, &SHLODOutliner::OnLevelSelectionChanged);
@@ -904,7 +926,7 @@ namespace HLODOutliner
 		// Register to update when an undo/redo operation has been called to update our list of actors
 		GEditor->RegisterForUndo(this);
 
-		RegisterActiveTimer(0.1f, FWidgetActiveTimerDelegate::CreateSP(this, &SHLODOutliner::UpdateNeedsBuildFlagTimer));
+		ActiveTimerHandle = RegisterActiveTimer(0.1f, FWidgetActiveTimerDelegate::CreateSP(this, &SHLODOutliner::UpdateNeedsBuildFlagTimer));
 	}
 
 	void SHLODOutliner::DeregisterDelegates()
@@ -915,6 +937,7 @@ namespace HLODOutliner
 
 		FWorldDelegates::LevelAddedToWorld.RemoveAll(this);
 		FWorldDelegates::LevelRemovedFromWorld.RemoveAll(this);
+
 		USelection::SelectionChangedEvent.RemoveAll(this);
 		USelection::SelectObjectEvent.RemoveAll(this);
 
@@ -930,12 +953,15 @@ namespace HLODOutliner
 		{
 			GEditor->OnHLODActorMoved().RemoveAll(this);
 			GEditor->OnHLODActorAdded().RemoveAll(this);
+			GEditor->OnHLODTransitionScreenSizeChanged().RemoveAll(this);
 			GEditor->OnHLODLevelsArrayChanged().RemoveAll(this);
 			GEditor->OnHLODActorRemovedFromCluster().RemoveAll(this);
 
 			// Deregister for Undo callbacks
 			GEditor->UnregisterForUndo(this);
 		}
+
+		UnRegisterActiveTimer(ActiveTimerHandle.ToSharedRef());
 	}
 
 	void SHLODOutliner::ForceViewLODActor()
@@ -1882,7 +1908,7 @@ namespace HLODOutliner
 		checkf(bUpdatedWorld == true, TEXT("Could not find UWorld* instance in Engine world contexts"));
 
 		TArray<FTreeItemRef> LevelNodes;
-		if (CurrentWorldSettings)
+		if (OutlinerEnabled())
 		{
 			// Iterate over all LOD levels (Number retrieved from world settings) and add Treeview items for them
 			const TArray<struct FHierarchicalSimplification>& HierarchicalLODSetups = CurrentWorldSettings->GetHierarchicalLODSetup();
@@ -2116,18 +2142,21 @@ namespace HLODOutliner
 	{
 		bool bHLODEnabled = false;
 
-		if (!bNeedsRefresh)
+		if (CurrentWorldSettings != nullptr)
 		{
-			if (CurrentWorldSettings != nullptr)
-			{
-				bHLODEnabled = CurrentWorldSettings->bEnableHierarchicalLODSystem;
-			}
+			bHLODEnabled = CurrentWorldSettings->bEnableHierarchicalLODSystem;
+		}
 
-			if (bHLODEnabled && CurrentWorld.IsValid())
-			{
-				bHLODEnabled = !HierarchicalLODUtilities->IsWorldUsedForStreaming(CurrentWorld.Get());
-			}
-		}	
+		if (bHLODEnabled && CurrentWorld.IsValid())
+		{
+			bHLODEnabled = !HierarchicalLODUtilities->IsWorldUsedForStreaming(CurrentWorld.Get());
+		}
+
+		// Disable outliner in PIE
+		if (bHLODEnabled && CurrentWorld.IsValid())
+		{
+			bHLODEnabled = !CurrentWorld->IsPlayInEditor();
+		}
 
 		return bHLODEnabled;
 	}
diff --git a/Engine/Source/Editor/HierarchicalLODOutliner/Private/HLODOutliner.h b/Engine/Source/Editor/HierarchicalLODOutliner/Private/HLODOutliner.h
--- a/Engine/Source/Editor/HierarchicalLODOutliner/Private/HLODOutliner.h
+++ b/Engine/Source/Editor/HierarchicalLODOutliner/Private/HLODOutliner.h
@@ -358,6 +358,13 @@ namespace HLODOutliner
 
 		/** Callback function used to check if Hierarchical LOD functionality is enabled in the current world settings */
 		bool OutlinerEnabled() const;
+
+		/** Called when a PIE session is beginning */
+		void OnBeginPieEvent(bool bIsSimulating);
+
+		/** Called when a PIE session is ending */
+		void OnEndPieEvent(bool bIsSimulating);
+
 	private:
 		/** Tells the scene outliner that it should do a full refresh, which will clear the entire tree and rebuild it from scratch. */
 		void FullRefresh();
@@ -465,6 +472,9 @@ namespace HLODOutliner
 		/** Cached pointer to HLOD utilities */
 		IHierarchicalLODUtilities* HierarchicalLODUtilities;
 
+		/** Update active timer handle */
+		TSharedPtr<FActiveTimerHandle> ActiveTimerHandle;
+
 		/** Cached flag to see if we need to generate meshes (any actors are dirty) */
 		bool bCachedNeedsBuild;
 
