Remove from Cluster option is missing when right clicking on generated actor in Hierarchical LOD Outliner

diff --git a/Engine/Source/Editor/HierarchicalLODOutliner/Private/StaticMeshActorItem.cpp b/Engine/Source/Editor/HierarchicalLODOutliner/Private/StaticMeshActorItem.cpp
--- a/Engine/Source/Editor/HierarchicalLODOutliner/Private/StaticMeshActorItem.cpp
+++ b/Engine/Source/Editor/HierarchicalLODOutliner/Private/StaticMeshActorItem.cpp
@@ -25,7 +25,7 @@ void HLODOutliner::FStaticMeshActorItem::GenerateContextMenu(UToolMenu* Menu, SH
 	auto SharedOutliner = StaticCastSharedRef<SHLODOutliner>(Outliner.AsShared());
 	FToolMenuSection& Section = Menu->AddSection("Section");
 	Section.AddMenuEntry("RemoveSMActorFromCluster", LOCTEXT("RemoveSMActorFromCluster", "Remove From Cluster"), FText(), FSlateIcon(), FUIAction(FExecuteAction::CreateRaw(&Outliner, &SHLODOutliner::RemoveStaticMeshActorFromCluster)));
-	Section.AddMenuEntry("RemoveSMActorFromCluster", LOCTEXT("ExcludeSMActorFromGeneration", "Exclude From Cluster Generation"), FText(), FSlateIcon(), FUIAction(FExecuteAction::CreateRaw(&Outliner, &SHLODOutliner::ExcludeFromClusterGeneration)));
+	Section.AddMenuEntry("ExcludeSMActorFromGeneration", LOCTEXT("ExcludeSMActorFromGeneration", "Exclude From Cluster Generation"), FText(), FSlateIcon(), FUIAction(FExecuteAction::CreateRaw(&Outliner, &SHLODOutliner::ExcludeFromClusterGeneration)));
 }
 
 FString HLODOutliner::FStaticMeshActorItem::GetDisplayString() const
