AutoSave detects dirty assets after saving changes made to a level containing a HLOD - causing a crash

On actor deletion, if actor is a LODActor, make sure it isn't part of the selection array

diff --git a/Engine/Source/Editor/HierarchicalLODOutliner/Private/HLODOutliner.cpp b/Engine/Source/Editor/HierarchicalLODOutliner/Private/HLODOutliner.cpp
--- a/Engine/Source/Editor/HierarchicalLODOutliner/Private/HLODOutliner.cpp
+++ b/Engine/Source/Editor/HierarchicalLODOutliner/Private/HLODOutliner.cpp
@@ -1702,7 +1702,13 @@ namespace HLODOutliner
 	void SHLODOutliner::OnLevelActorsRemoved(AActor* InActor)
 	{
 		if (!InActor->IsA<AWorldSettings>())
-		{			
+		{
+			// If actor is a LODActor, make sure it isn't part of the selection array.
+			if (InActor->IsA<ALODActor>())
+			{
+				SelectedLODActors.Remove(InActor);
+			}
+
 			// Remove InActor from LOD actor which contains it
 			for (TArray<TWeakObjectPtr<ALODActor>>& ActorArray : LODLevelActors)
 			{				
