Update HLOD Outliner to new UX style

diff --git a/Engine/Source/Editor/HierarchicalLODOutliner/Private/HLODOutliner.cpp b/Engine/Source/Editor/HierarchicalLODOutliner/Private/HLODOutliner.cpp
--- a/Engine/Source/Editor/HierarchicalLODOutliner/Private/HLODOutliner.cpp
+++ b/Engine/Source/Editor/HierarchicalLODOutliner/Private/HLODOutliner.cpp
@@ -45,7 +45,7 @@
 #include "Widgets/Layout/SScrollBorder.h"
 #include "Widgets/Layout/SWidgetSwitcher.h"
 #include "Widgets/Input/SComboButton.h"
-#include "Widgets/Layout/SWrapBox.h"
+#include "Styling/StyleColors.h"
 #include "Framework/Notifications/NotificationManager.h"
 #include "Widgets/Notifications/SNotificationList.h"
 
@@ -63,6 +63,7 @@ namespace HLODOutliner
 	SHLODOutliner::SHLODOutliner()		
 	{
 		bNeedsRefresh = true;
+		ToolBarButtonCurrentState = GenerateClusters;
 		CurrentWorld = nullptr;
 		CurrentWorldSettings = nullptr;
 		ForcedLODLevel = LOD_AUTO;
@@ -90,7 +91,7 @@ namespace HLODOutliner
 	}
 
 	BEGIN_SLATE_FUNCTION_BUILD_OPTIMIZATION
-		void SHLODOutliner::Construct(const FArguments& InArgs)
+	void SHLODOutliner::Construct(const FArguments& InArgs)
 	{
 		CreateSettingsView();
 
@@ -147,11 +148,10 @@ namespace HLODOutliner
 
 		MainContentPanel->AddSlot()
 			.AutoHeight()
-			.Padding(0.0f, 0.0f, 0.0f, 4.0f)
 			[
 				SNew(SBorder)
 				.BorderImage(FCoreStyle::Get().GetBrush("ToolPanel.GroupBorder"))
-				.Padding(1.0f)
+				.Padding(FMargin(4.0f, 4.0f, 4.0f, 0.0f))
 				[
 					CreateMainButtonWidgets()
 				]
@@ -161,12 +161,16 @@ namespace HLODOutliner
 
 		TSharedRef<SWidget> ClusterWidgets = 
 			SNew(SBorder)
-			.BorderImage(FCoreStyle::Get().GetBrush("ToolPanel.GroupBorder"))
-			.Padding(1.0f)
+			.BorderImage(&FEditorStyle::GetWidgetStyle<FTableRowStyle>("TableView.Row").EvenRowBackgroundBrush)
+			.Padding(0)
 			[
 				SNew(SVerticalBox)
 				+SVerticalBox::Slot()
-				.Padding(2.0f)
+				.AutoHeight()
+				[
+					CreateClusterButtonWidgets()
+				]
+				+SVerticalBox::Slot()
 				.FillHeight(1.0f)
 				[
 					SNew(SScrollBorder, TreeViewWidget)
@@ -174,12 +178,6 @@ namespace HLODOutliner
 						TreeViewWidget
 					]
 				]
-				+SVerticalBox::Slot()
-				.AutoHeight()
-				.HAlign(HAlign_Right)
-				[
-					CreateClusterButtonWidgets()
-				]
 			];
 
 		TSharedRef<SWidget> DetailsWidgets = 
@@ -218,6 +216,7 @@ namespace HLODOutliner
 							SNew(STextBlock)
 							.AutoWrapText(true)
 							.Text(LOCTEXT("HLODForcedGlobally", "Project level HLOD Settings forced, changing the HLOD settings is disabled"))
+							.ColorAndOpacity(FSlateColor(EStyleColor::Black))
 						]
 					]
 				]
@@ -236,6 +235,7 @@ namespace HLODOutliner
 				+SWidgetSwitcher::Slot()
 				[
 					SNew(SSplitter)
+					.Style(FEditorStyle::Get(), "DetailsView.Splitter")
 					.Orientation(Orient_Horizontal)
 					+ SSplitter::Slot()
 					.Value(0.5)
@@ -251,6 +251,7 @@ namespace HLODOutliner
 				+SWidgetSwitcher::Slot()
 				[
 					SNew(SSplitter)
+					.Style(FEditorStyle::Get(), "DetailsView.Splitter")
 					.Orientation(Orient_Vertical)
 					+ SSplitter::Slot()
 					.Value(0.5)
@@ -296,6 +297,26 @@ namespace HLODOutliner
 						SNew(STextBlock)
 						.AutoWrapText(true)
 						.Text(LOCTEXT("HLODNeedsBuild", "Actors represented in HLOD have changed, generate proxy meshes to update."))
+						.ColorAndOpacity(FSlateColor(EStyleColor::Black))
+					]
+				]
+			];
+
+		MainContentPanel->AddSlot()
+			.AutoHeight()
+			[
+				SNew(SBorder)
+				.BorderImage(FEditorStyle::GetBrush("ToolPanel.GroupBorder"))
+				.Padding(FMargin(8.0f, 8.0f, 0.0f, 8.0f))
+				[
+					SNew(SCheckBox)
+					.Type(ESlateCheckBoxType::CheckBox)
+					.IsChecked_Lambda([this]() { return (CurrentWorldSettings && CurrentWorldSettings->bGenerateSingleClusterForLevel) ? ECheckBoxState::Checked : ECheckBoxState::Unchecked; })
+					.OnCheckStateChanged_Lambda([this](ECheckBoxState NewState) {  if (CurrentWorldSettings) { CurrentWorldSettings->bGenerateSingleClusterForLevel = (NewState == ECheckBoxState::Checked); } })
+					.Content()
+					[
+						SNew(STextBlock)
+						.Text(LOCTEXT("GenerateSingleClusterLabel", "Generate Single Cluster for Level"))
 					]
 				]
 			];
@@ -309,110 +330,54 @@ namespace HLODOutliner
 		FEditorDelegates::EndPIE.AddRaw(this, &SHLODOutliner::OnEndPieEvent);
 	}
 
+	TSharedRef<SWidget> SHLODOutliner::MakeToolBar()
+	{
+		FSlimHorizontalToolBarBuilder ToolBarBuilder(nullptr, FMultiBoxCustomization::None);
+
+		ToolBarBuilder.AddToolBarButton(
+			FUIAction(
+				FExecuteAction::CreateSP(this, &SHLODOutliner::OnToolBarButtonClicked),
+				FCanExecuteAction::CreateSP(this, &SHLODOutliner::IsToolBarButtonEnabled)),
+			NAME_None,
+			TAttribute<FText>::Create(TAttribute<FText>::FGetter::CreateSP(this, &SHLODOutliner::GetToolBarButtonLabel)),
+			TAttribute<FText>::Create(TAttribute<FText>::FGetter::CreateSP(this, &SHLODOutliner::GetToolBarButtonToolTip)),
+			TAttribute<FSlateIcon>::Create(TAttribute<FSlateIcon>::FGetter::CreateSP(this, &SHLODOutliner::GetToolBarButtonIcon))
+		);
+
+		ToolBarBuilder.AddComboButton(
+			FUIAction(),
+			FOnGetContent::CreateRaw(this, &SHLODOutliner::GetToolBarButtonMenu),
+			TAttribute<FText>(),
+			TAttribute<FText>(),
+			TAttribute<FSlateIcon>(),
+			true);
+
+		ToolBarBuilder.AddToolBarButton(
+			FUIAction(FExecuteAction::CreateLambda([this]() { HandleSaveAll(); })),
+			NAME_None,
+			FText(),
+			LOCTEXT("SaveAllToolTip", "Saves all external HLOD data: Meshes, materials etc."),
+			FSlateIcon(FEditorStyle::GetStyleSetName(), "AssetEditor.SaveAsset")
+		);
+
+		return ToolBarBuilder.MakeWidget();
+	}
+
 	TSharedRef<SWidget> SHLODOutliner::CreateMainButtonWidgets()
 	{
 		return SNew(SVerticalBox)
 			+SVerticalBox::Slot()
-			.Padding(FMargin(0.0f, 2.0f))
+			.Padding(FMargin(2.0f, 4.0f))
 			[
-				SNew(SWrapBox)
-				.UseAllottedSize(true)
-
-				+ SWrapBox::Slot()
-				.Padding(FMargin(2.0f))
-				[
-					SNew(SButton)
-					.ButtonStyle(FEditorStyle::Get(), "FlatButton.Success")
-					.HAlign(HAlign_Center)
-					.OnClicked(this, &SHLODOutliner::HandleBuildLODActors)
-					.IsEnabled(this, &SHLODOutliner::CanBuildLODActors)
-					.ToolTipText(this, &SHLODOutliner::GetBuildLODActorsTooltipText)
-					[
-						SNew( SHorizontalBox )
-						+ SHorizontalBox::Slot()
-						.VAlign(VAlign_Center)
-						.AutoWidth()
-						[
-							SNew(STextBlock)
-							.TextStyle(FEditorStyle::Get(), "ContentBrowser.TopBar.Font")
-							.Font(FEditorStyle::Get().GetFontStyle("FontAwesome.11"))
-							.Text(FEditorFontGlyphs::Building)
-						]
-						+ SHorizontalBox::Slot()
-						.AutoWidth()
-						.VAlign(VAlign_Center)
-						.Padding(4, 0, 0, 0)
-						[
-							SNew( STextBlock )
-							.TextStyle( FEditorStyle::Get(), "ContentBrowser.TopBar.Font" )
-							.Text(this, &SHLODOutliner::GetBuildText)
-						]
-					]
-				]
-
-				+ SWrapBox::Slot()
-				.Padding(FMargin(2.0f))
-				[
-					SNew(SButton)
-					.ButtonStyle(FEditorStyle::Get(), "FlatButton.Danger")
-					.HAlign(HAlign_Center)
-					.OnClicked(this, &SHLODOutliner::HandleForceBuildLODActors)
-					.ToolTipText(LOCTEXT("BuildClustersAndMeshesToolTip", "Re-generates clusters and then proxy meshes for each of the generated clusters in the level. This dirties the level."))
-					[
-						SNew( SHorizontalBox )
-						+ SHorizontalBox::Slot()
-						.VAlign(VAlign_Center)
-						.AutoWidth()
-						[
-							SNew(STextBlock)
-							.TextStyle(FEditorStyle::Get(), "ContentBrowser.TopBar.Font")
-							.Font(FEditorStyle::Get().GetFontStyle("FontAwesome.11"))
-							.Text(FEditorFontGlyphs::Recycle)
-						]
-						+ SHorizontalBox::Slot()
-						.AutoWidth()
-						.VAlign(VAlign_Center)
-						.Padding(4, 0, 0, 0)
-						[
-							SNew( STextBlock )
-							.TextStyle( FEditorStyle::Get(), "ContentBrowser.TopBar.Font" )
-							.Text(this, &SHLODOutliner::GetForceBuildText)
-						]
-					]
-				]
-
-				+ SWrapBox::Slot()
+				SNew(SHorizontalBox)
+				+SHorizontalBox::Slot()
 				.Padding(FMargin(2.0f))
+				.AutoWidth()
 				[
-					SNew(SButton)
-					.ButtonStyle(FEditorStyle::Get(), "FlatButton")
-					.HAlign(HAlign_Center)
-					.OnClicked(this, &SHLODOutliner::HandleSaveAll)
-					.ToolTipText(LOCTEXT("SaveAllToolTip", "Saves all external HLOD data: Meshes, materials etc."))
-					[
-						SNew( SHorizontalBox )
-						+ SHorizontalBox::Slot()
-						.VAlign(VAlign_Center)
-						.AutoWidth()
-						[
-							SNew(STextBlock)
-							.TextStyle(FEditorStyle::Get(), "ContentBrowser.TopBar.Font")
-							.Font(FEditorStyle::Get().GetFontStyle("FontAwesome.11"))
-							.Text(FEditorFontGlyphs::Floppy_O)
-						]
-						+ SHorizontalBox::Slot()
-						.AutoWidth()
-						.VAlign(VAlign_Center)
-						.Padding(4, 0, 0, 0)
-						[
-							SNew( STextBlock )
-							.TextStyle( FEditorStyle::Get(), "ContentBrowser.TopBar.Font" )
-							.Text(LOCTEXT("SaveAll", "Save All"))
-						]
-					]
+					MakeToolBar()
 				]
-
-				+ SWrapBox::Slot()
+				+SHorizontalBox::Slot()
+				.HAlign(HAlign_Right)
 				.Padding(FMargin(2.0f))
 				[
 					CreateForcedViewWidget()
@@ -424,73 +389,62 @@ namespace HLODOutliner
 	{
 		return SNew(SVerticalBox)
 			+ SVerticalBox::Slot()
-			.Padding(FMargin(0.0f, 2.0f))
+			.Padding(FMargin(0.0f, 0.0f, 0.0f, 2.0f))
 			[
-				SNew(SHorizontalBox)
-
-				+ SHorizontalBox::Slot()
-				.Padding(FMargin(2.0f))
-				.VAlign(VAlign_Center)
-				.AutoWidth()
+				SNew(SBorder)
+				.BorderImage(FEditorStyle::GetBrush("DetailsView.CategoryTop"))
 				[
-					SNew(SButton)
-					.ButtonStyle(FEditorStyle::Get(), "FlatButton")
-					.HAlign(HAlign_Center)
-					.OnClicked(this, &SHLODOutliner::HandlePreviewHLODs)
-					.ToolTipText(LOCTEXT("GenerateClusterToolTip", "Generates clusters (but not proxy meshes) for meshes in the level"))
+					SNew(SHorizontalBox)
+					+ SHorizontalBox::Slot()
+					.HAlign(HAlign_Left)
+					.VAlign(VAlign_Center)
+					.Padding(FMargin(6.0f, 0.0f, 0.0f, 0.0f))
 					[
-						SNew( SHorizontalBox )
-						+ SHorizontalBox::Slot()
-						.VAlign(VAlign_Center)
-						.AutoWidth()
-						[
-							SNew(STextBlock)
-							.TextStyle(FEditorStyle::Get(), "ContentBrowser.TopBar.Font")
-							.Font(FEditorStyle::Get().GetFontStyle("FontAwesome.11"))
-							.Text(FEditorFontGlyphs::List)
-						]
-						+ SHorizontalBox::Slot()
-						.AutoWidth()
-						.VAlign(VAlign_Center)
-						.Padding(4, 0, 0, 0)
-						[
-							SNew( STextBlock )
-							.TextStyle( FEditorStyle::Get(), "ContentBrowser.TopBar.Font" )
-							.Text(LOCTEXT("GenerateClusters", "Generate Clusters"))
-						]
+						SNew(STextBlock)
+						.TextStyle(FEditorStyle::Get(), "DialogButtonText")
+						.Text(LOCTEXT("ClustersLabel", "Clusters"))
 					]
-				]
 
-				+ SHorizontalBox::Slot()
-				.Padding(FMargin(2.0f))
-				.VAlign(VAlign_Center)
-				.AutoWidth()
-				[
-					SNew(SButton)
-					.ButtonStyle(FEditorStyle::Get(), "FlatButton")
-					.HAlign(HAlign_Center)
-					.OnClicked(this, &SHLODOutliner::HandleDeleteHLODs)
-					.IsEnabled(this, &SHLODOutliner::CanDeleteHLODs)
-					.ToolTipText(LOCTEXT("DeleteClusterToolTip", "Deletes all clusters in the level"))
+					+ SHorizontalBox::Slot()
+					.Padding(FMargin(2.0f))
+					.HAlign(HAlign_Right)
+					.VAlign(VAlign_Center)
+					.AutoWidth()
 					[
-						SNew( SHorizontalBox )
-						+ SHorizontalBox::Slot()
-						.VAlign(VAlign_Center)
-						.AutoWidth()
+						SNew(SButton)
+						.ButtonStyle(FEditorStyle::Get(), "SimpleButton")
+						.ContentPadding(0)
+						.HAlign(HAlign_Center)
+						.OnClicked(this, &SHLODOutliner::GenerateClustersFromUI)
+						.ToolTipText(LOCTEXT("GenerateClusterToolTip", "Generates clusters (but not proxy meshes) for meshes in the level"))
 						[
-							SNew(STextBlock)
-							.TextStyle(FEditorStyle::Get(), "ContentBrowser.TopBar.Font")
-							.Font(FEditorStyle::Get().GetFontStyle("FontAwesome.11"))
-							.Text(FEditorFontGlyphs::Trash)
+							SNew(SImage)
+							.Image(FEditorStyle::Get().GetBrush("Icons.PlusCircle"))
 						]
-						+ SHorizontalBox::Slot()
-						.AutoWidth()
-						.VAlign(VAlign_Center)
-						.Padding(4, 0, 0, 0)
+					]
+
+					+ SHorizontalBox::Slot()
+					.Padding(FMargin(2.0f))
+					.HAlign(HAlign_Right)
+					.VAlign(VAlign_Center)
+					.AutoWidth()
+					[
+						SNew(SButton)
+						.ButtonStyle(FEditorStyle::Get(), "SimpleButton")
+						.ContentPadding(0)
+						.HAlign(HAlign_Center)
+						.OnClicked_Lambda([this]() {
+							// Transition out of the generate proxy meshes state if we were in it.
+							if(ToolBarButtonCurrentState == GenerateProxyMeshes)
+								ToolBarButtonCurrentState = GenerateClusters;
+
+							return HandleDeleteHLODs();
+							})
+						.IsEnabled(this, &SHLODOutliner::CanDeleteHLODs)
+						.ToolTipText(LOCTEXT("DeleteClusterToolTip", "Deletes all clusters in the level"))
 						[
-							SNew( STextBlock )
-							.TextStyle( FEditorStyle::Get(), "ContentBrowser.TopBar.Font" )
-							.Text(LOCTEXT("DeleteClusters", "Delete Clusters"))
+							SNew(SImage)
+							.Image(FEditorStyle::Get().GetBrush("Icons.Delete"))
 						]
 					]
 				]
@@ -512,24 +466,29 @@ namespace HLODOutliner
 			(
 				SNew(SHeaderRow)
 				+ SHeaderRow::Column("SceneActorName")
-				.DefaultLabel(LOCTEXT("SceneActorName", "Scene Actor Name"))
+				.DefaultLabel(LOCTEXT("SceneActorName", "Actor Name"))
+				.HeaderContentPadding(FMargin(8.0f, 0.0f, 0.0f, 0.0f))
 				.FillWidth(0.3f)				
 				+ SHeaderRow::Column("RawTriangleCount")
-				.DefaultLabel(LOCTEXT("RawTriangleCount", "Original Triangle Count"))
+				.DefaultLabel(LOCTEXT("RawTriangleCount", "Tri Original"))
 				.DefaultTooltip(LOCTEXT("RawTriangleCountToolTip", "Original Number of Triangles in a LOD Mesh"))
+				.HeaderContentPadding(FMargin(8.0f, 0.0f, 0.0f, 0.0f))
 				.FillWidth(0.2f)
 				+ SHeaderRow::Column("ReducedTriangleCount")
-				.DefaultLabel(LOCTEXT("ReducedTriangleCount", "Reduced Triangle Count"))
+				.DefaultLabel(LOCTEXT("ReducedTriangleCount", "Tri Reduced"))
 				.DefaultTooltip(LOCTEXT("ReducedTriangleCountToolTip", "Reduced Number of Triangles in a LOD Mesh"))
+				.HeaderContentPadding(FMargin(8.0f, 0.0f, 0.0f, 0.0f))
 				.FillWidth(0.2f)
 				+ SHeaderRow::Column("ReductionPercentage")
 				.DefaultLabel(LOCTEXT("ReductionPercentage", "% Retained"))
 				.DefaultTooltip(LOCTEXT("ReductionPercentageToolTip", "Percentage of Triangle Reduction in a LOD Mesh"))
-				.FillWidth(0.1f)
+				.HeaderContentPadding(FMargin(8.0f, 0.0f, 0.0f, 0.0f))
+				.FillWidth(0.2f)
 				+ SHeaderRow::Column("Level")
 				.DefaultLabel(LOCTEXT("Level", "Level"))
 				.DefaultTooltip(LOCTEXT("LevelToolTip", "Persistent Level of a LOD Mesh"))
-				.FillWidth(0.2f)
+				.HeaderContentPadding(FMargin(8.0f, 0.0f, 0.0f, 0.0f))
+				.FillWidth(0.1f)
 			);
 	}
 
@@ -539,14 +498,13 @@ namespace HLODOutliner
 				.ContentPadding(FMargin(4.0f, 2.0f))
 				.ForegroundColor(FLinearColor::White)
 				.ButtonStyle(FEditorStyle::Get(), "FlatButton")
-				.ComboButtonStyle(FEditorStyle::Get(), "ToolbarComboButton")
 				.HasDownArrow(true)
 				.OnGetMenuContent(this, &SHLODOutliner::GetForceLevelMenuContent)
 				.ToolTipText(LOCTEXT("ForcedLODButtonTooltip", "Choose the LOD level to view."))
 				.ButtonContent()
 				[
 					SNew(STextBlock)
-					.TextStyle(FEditorStyle::Get(), "ContentBrowser.TopBar.Font")
+					.TextStyle(FEditorStyle::Get(), "DialogButtonText")
 					.Text(this, &SHLODOutliner::HandleForceLevelText)
 				];
 	}
@@ -581,6 +539,15 @@ namespace HLODOutliner
 				};
 
 				FString CategoryName = PropertyAndParent.Property.GetMetaData("Category");
+				
+				// Exceptions
+				// This one is shown at the bottom of the windows, outside of the properties
+				if (CategoryName == "HLODSystem" && PropertyAndParent.Property.GetName() == "bGenerateSingleClusterForLevel")
+				{
+					return false;
+				}
+
+				// General case
 				for (uint32 CategoryIndex = 0; CategoryIndex < 5; ++CategoryIndex)
 				{
 					if (CategoryName == CategoryNames[CategoryIndex])
@@ -652,7 +619,7 @@ namespace HLODOutliner
 
 			// Restore expansion states
 			TreeView->RequestTreeRefresh();		
-		}			
+		}
 
 		bArrangeHorizontally = AllottedGeometry.Size.X > AllottedGeometry.Size.Y;
 	}
@@ -723,22 +690,12 @@ namespace HLODOutliner
 
 	FText SHLODOutliner::GetForceBuildText() const
 	{
-		return HasHLODActors() ? LOCTEXT("RebuildAllClustersAndMeshes", "Rebuild All") : LOCTEXT("BuildClustersAndMeshes", "Build");
+		return HasHLODActors() ? LOCTEXT("RebuildAllClustersAndMeshes", "Rebuild All") : LOCTEXT("BuildClustersAndMeshes", "Build Meshes");
 	}
 
-	FReply SHLODOutliner::HandleBuildHLODs()
+	FText SHLODOutliner::GetForceBuildToolTip() const
 	{
-		CloseOpenAssetEditors();
-
-		if (CurrentWorld.IsValid())
-		{
-			CurrentWorld->HierarchicalLODBuilder->Build();
-		}
-
-		FMessageLog("HLODResults").Open();
-
-		FullRefresh();
-		return FReply::Handled();
+		return LOCTEXT("BuildClustersAndMeshesToolTip", "Re-generates clusters and then proxy meshes for each of the generated clusters in the level. This dirties the level.");
 	}
 
 	FReply SHLODOutliner::HandleDeleteHLODs()
@@ -774,16 +731,6 @@ namespace HLODOutliner
 		return FReply::Handled();
 	}
 
-	FReply SHLODOutliner::HandleDeletePreviewHLODs()
-	{
-		if (CurrentWorld.IsValid())
-		{
-			CurrentWorld->HierarchicalLODBuilder->ClearPreviewBuild();
-		}
-		FullRefresh();
-		return FReply::Handled();
-	}
-
 	FReply SHLODOutliner::HandleBuildLODActors()
 	{
 		if (CurrentWorld.IsValid())
@@ -870,8 +817,187 @@ namespace HLODOutliner
 		return FReply::Handled();
 	}
 
+	TSharedRef<SWidget> SHLODOutliner::GetToolBarButtonMenu()
+	{
+		const bool bInShouldCloseWindowAfterMenuSelection = true;
+		FMenuBuilder MenuBuilder(bInShouldCloseWindowAfterMenuSelection, nullptr);
+
+		if (HasHLODActors())
+		{
+			MenuBuilder.AddMenuEntry(
+				GetRegenerateClustersText(),
+				GetRegenerateClustersTooltip(),
+				FSlateIcon(FEditorStyle::GetStyleSetName(), "GenericCommands.Redo"),
+				FUIAction(FExecuteAction::CreateLambda([this]() { RegenerateClustersFromUI(); })));
+		}
+		else
+		{
+			MenuBuilder.AddMenuEntry(
+				GetGenerateClustersText(),
+				GetRegenerateClustersTooltip(),
+				FSlateIcon(FEditorStyle::GetStyleSetName(), "Icons.PlusCircle"),
+				FUIAction(FExecuteAction::CreateLambda([this]() { GenerateClustersFromUI(); })));
+		}
+
+		MenuBuilder.AddMenuEntry(
+			GetBuildText(),
+			GetBuildLODActorsTooltipText(),
+			FSlateIcon(),
+			FUIAction(
+				FExecuteAction::CreateLambda([this](){ GenerateProxyMeshesFromUI(); }),
+				FCanExecuteAction::CreateSP(this, &SHLODOutliner::CanBuildLODActors)));
+
+		MenuBuilder.AddMenuEntry(
+			GetForceBuildText(),
+			GetForceBuildToolTip(),
+			FSlateIcon(),
+			FUIAction(FExecuteAction::CreateLambda([this] { BuildClustersAndMeshesFromUI(); })));
+
+		return MenuBuilder.MakeWidget();
+	}
+
+	FText SHLODOutliner::GetToolBarButtonLabel()
+	{
+		switch (ToolBarButtonCurrentState)
+		{
+		default:
+		case GenerateClusters:
+			if (HasHLODActors())
+			{
+				return GetRegenerateClustersText();
+			}
+			else
+			{
+				return GetGenerateClustersText();
+			}
+
+		case GenerateProxyMeshes:
+			return GetBuildText();
+
+		case RebuildAll:
+			return GetForceBuildText();
+		}
+	}
+
+	FText SHLODOutliner::GetToolBarButtonToolTip()
+	{
+		switch (ToolBarButtonCurrentState)
+		{
+		default:
+		case GenerateClusters:
+			if (HasHLODActors())
+			{
+				return GetRegenerateClustersText();
+			}
+			else
+			{
+				return GetGenerateClustersText();
+			}
+
+		case GenerateProxyMeshes:
+			return GetBuildLODActorsTooltipText();
+
+		case RebuildAll:
+			return GetForceBuildToolTip();
+		}
+	}
+
+	FSlateIcon SHLODOutliner::GetToolBarButtonIcon()
+	{
+		switch (ToolBarButtonCurrentState)
+		{
+		case GenerateClusters:
+			if (HasHLODActors())
+			{
+				return FSlateIcon(FEditorStyle::GetStyleSetName(), "GenericCommands.Redo");
+			}
+			else
+			{
+				return FSlateIcon(FEditorStyle::GetStyleSetName(), "Icons.PlusCircle");
+			}
+
+		case GenerateProxyMeshes:
+		case RebuildAll:
+		default:
+			return FSlateIcon();
+		}
+	}
+
+	void SHLODOutliner::OnToolBarButtonClicked()
+	{
+		switch (ToolBarButtonCurrentState)
+		{
+		default:
+		case GenerateClusters:
+			if (HasHLODActors())
+			{
+				RegenerateClustersFromUI();
+			}
+			else
+			{
+				GenerateClustersFromUI();
+			}
+			break;
+		case GenerateProxyMeshes:
+			GenerateProxyMeshesFromUI();
+			break;
+		case RebuildAll:
+			BuildClustersAndMeshesFromUI();
+			break;
+		}
+	}
+
+	bool SHLODOutliner::IsToolBarButtonEnabled()
+	{
+		return (ToolBarButtonCurrentState != GenerateProxyMeshes) || CanBuildLODActors();
+	}
+
 	END_SLATE_FUNCTION_BUILD_OPTIMIZATION
 
+	FReply SHLODOutliner::RegenerateClustersFromUI()
+	{
+		ToolBarButtonCurrentState = GenerateClusters;
+		return HandlePreviewHLODs();
+	}
+
+	FReply SHLODOutliner::GenerateClustersFromUI()
+	{
+		ToolBarButtonCurrentState = GenerateProxyMeshes;
+		return HandlePreviewHLODs();
+	}
+
+	FReply SHLODOutliner::GenerateProxyMeshesFromUI()
+	{
+		ToolBarButtonCurrentState = GenerateProxyMeshes;
+		return HandleBuildLODActors();
+	}
+
+	FReply SHLODOutliner::BuildClustersAndMeshesFromUI()
+	{
+		ToolBarButtonCurrentState = RebuildAll;
+		return HandleForceBuildLODActors();
+	}
+
+	FText SHLODOutliner::GetGenerateClustersText() const
+	{
+		return LOCTEXT("GenerateClustersLabel", "Generate Clusters");
+	}
+
+	FText SHLODOutliner::GetGenerateClustersTooltip() const
+	{
+		return LOCTEXT("GenerateClusterToolTip", "Generates clusters (but not proxy meshes) for meshes in the level");
+	}
+
+	FText SHLODOutliner::GetRegenerateClustersText() const
+	{
+		return LOCTEXT("RegenerateClustersLabel", "Regenerate Clusters");
+	}
+
+	FText SHLODOutliner::GetRegenerateClustersTooltip() const
+	{
+		return LOCTEXT("RegenerateClusterToolTip", "Regenerates clusters (but not proxy meshes) for meshes in the level");
+	}
+
 	void SHLODOutliner::OnBeginPieEvent(bool bIsSimulating)
 	{
 		DeregisterDelegates();
diff --git a/Engine/Source/Editor/HierarchicalLODOutliner/Private/HLODOutliner.h b/Engine/Source/Editor/HierarchicalLODOutliner/Private/HLODOutliner.h
--- a/Engine/Source/Editor/HierarchicalLODOutliner/Private/HLODOutliner.h
+++ b/Engine/Source/Editor/HierarchicalLODOutliner/Private/HLODOutliner.h
@@ -107,12 +107,11 @@ namespace HLODOutliner
 
 		/** Button handlers */
 		FText GetForceBuildText() const;
+		FText GetForceBuildToolTip() const;
 		FText GetBuildText() const;
-		FReply HandleBuildHLODs();
 		FReply HandleDeleteHLODs();
 		bool CanDeleteHLODs() const;
 		FReply HandlePreviewHLODs();
-		FReply HandleDeletePreviewHLODs();
 		FReply RetrieveActors();
 		FReply HandleBuildLODActors();
 		bool CanBuildLODActors() const;
@@ -120,6 +119,10 @@ namespace HLODOutliner
 		FReply HandleForceBuildLODActors();
 		FReply HandleForceRefresh();
 		FReply HandleSaveAll();
+		FText GetGenerateClustersText() const;
+		FText GetGenerateClustersTooltip() const;
+		FText GetRegenerateClustersText() const;
+		FText GetRegenerateClustersTooltip() const;
 		/** End button handlers */
 
 	private:
@@ -129,6 +132,25 @@ namespace HLODOutliner
 		/** De-registers all the callback delegates required for keeping the treeview sync */
 		void DeregisterDelegates();
 
+		/** Builds the toolbar */
+		TSharedRef<SWidget> MakeToolBar();
+
+		/** Builds the toolbar combo menu */
+		TSharedRef<SWidget> GetToolBarButtonMenu();
+
+		/** Delegates to build the toolbar combo button */
+		FText GetToolBarButtonLabel();
+		FText GetToolBarButtonToolTip();
+		FSlateIcon GetToolBarButtonIcon();
+		void OnToolBarButtonClicked();
+		bool IsToolBarButtonEnabled();
+
+		/** Callbacks to make interaction from UI uniform vs. UI updates*/
+		FReply RegenerateClustersFromUI();
+		FReply GenerateClustersFromUI();
+		FReply GenerateProxyMeshesFromUI();
+		FReply BuildClustersAndMeshesFromUI();
+
 		/** Helper function used to determine whether there are any LOD actors present */
 		bool HasHLODActors() const;
 
@@ -446,7 +468,17 @@ namespace HLODOutliner
 		TSharedPtr<SVerticalBox> MainContentPanel;
 		/** Attribute determining if the outliner UI is enabled*/
 		TAttribute<bool> EnabledAttribute;
-	
+
+		enum ToolBarButtonState
+		{
+			GenerateClusters,
+			GenerateProxyMeshes,
+			RebuildAll
+		};
+
+		/** Dialog state to determine which button to display in the toolbar */
+		ToolBarButtonState ToolBarButtonCurrentState;
+
 		/** Map containing all the nodes with their corresponding keys */
 		TMultiMap<FTreeItemID, FTreeItemPtr> TreeItemsMap;
 
diff --git a/Engine/Source/Editor/HierarchicalLODOutliner/Private/HLODTreeWidgetItem.cpp b/Engine/Source/Editor/HierarchicalLODOutliner/Private/HLODTreeWidgetItem.cpp
--- a/Engine/Source/Editor/HierarchicalLODOutliner/Private/HLODTreeWidgetItem.cpp
+++ b/Engine/Source/Editor/HierarchicalLODOutliner/Private/HLODTreeWidgetItem.cpp
@@ -180,20 +180,24 @@ BEGIN_SLATE_FUNCTION_BUILD_OPTIMIZATION
 	}
 
 
-		TSharedRef<SWidget> SHLODWidgetItem::GenerateWidgetForColumn(const FName& ColumnName)
+	TSharedRef<SWidget> SHLODWidgetItem::GenerateWidgetForColumn(const FName& ColumnName)
 	{
+		const float DefaultColumnHorizontalPadding = 24.0f;
+		const float LastColumnHorizontalPadding = 8.0f;
+
 		if (ColumnName == TEXT("SceneActorName"))
 		{
 			return SNew(SHorizontalBox)
-
 				+ SHorizontalBox::Slot()
 				.AutoWidth()
+				.Padding(4.0f, 0.0f, 0.0f, 0.0f)
 				.VAlign(VAlign_Center)
 				[
 					SNew(SExpanderArrow, SharedThis(this))
+					.IndentAmount(20)
 				]
 
-			+ SHorizontalBox::Slot()
+				+ SHorizontalBox::Slot()
 				.AutoWidth()
 				.Padding(2.0f, 0.0f)
 				.VAlign(VAlign_Center)
@@ -210,6 +214,7 @@ BEGIN_SLATE_FUNCTION_BUILD_OPTIMIZATION
 			return SNew(SHorizontalBox)
 				+ SHorizontalBox::Slot()
 				.AutoWidth()
+				.Padding(DefaultColumnHorizontalPadding, 0.0f)
 				.VAlign(VAlign_Center)
 				[
 					SNew(STextBlock)
@@ -224,6 +229,7 @@ BEGIN_SLATE_FUNCTION_BUILD_OPTIMIZATION
 			return SNew(SHorizontalBox)
 				+ SHorizontalBox::Slot()
 				.AutoWidth()
+				.Padding(DefaultColumnHorizontalPadding, 0.0f)
 				.VAlign(VAlign_Center)
 				[
 					SNew(STextBlock)
@@ -238,6 +244,7 @@ BEGIN_SLATE_FUNCTION_BUILD_OPTIMIZATION
 			return SNew(SHorizontalBox)
 				+ SHorizontalBox::Slot()
 				.AutoWidth()
+				.Padding(DefaultColumnHorizontalPadding, 0.0f)
 				.VAlign(VAlign_Center)
 				[
 					SNew(STextBlock)
@@ -252,6 +259,7 @@ BEGIN_SLATE_FUNCTION_BUILD_OPTIMIZATION
 			return SNew(SHorizontalBox)
 				+ SHorizontalBox::Slot()
 				.AutoWidth()
+				.Padding(LastColumnHorizontalPadding, 0.0f)
 				.VAlign(VAlign_Center)
 				[
 					SNew(STextBlock)
@@ -259,8 +267,6 @@ BEGIN_SLATE_FUNCTION_BUILD_OPTIMIZATION
 					.ColorAndOpacity(this, &SHLODWidgetItem::GetTint)
 				];
 		}
-
-
 		else
 		{
 			return SNullWidget::NullWidget;
diff --git a/Engine/Source/Editor/HierarchicalLODOutliner/Private/LODActorItem.cpp b/Engine/Source/Editor/HierarchicalLODOutliner/Private/LODActorItem.cpp
--- a/Engine/Source/Editor/HierarchicalLODOutliner/Private/LODActorItem.cpp
+++ b/Engine/Source/Editor/HierarchicalLODOutliner/Private/LODActorItem.cpp
@@ -126,7 +126,7 @@ FText HLODOutliner::FLODActorItem::GetReductionPercentageAsText() const
 		const uint32 SubActorCount = LODActor->GetNumTrianglesInSubActors();
 		const uint32 MergedCount = LODActor->GetNumTrianglesInMergedMesh();
 		const float PercentageOfOriginal = ((float)MergedCount / (float)SubActorCount) * 100.0f;
-		return  FText::FromString(((SubActorCount != 0) ? FString::SanitizeFloat(PercentageOfOriginal) : TEXT("0")) + "%");
+		return  FText::FromString("%" + ((SubActorCount != 0) ? FString::SanitizeFloat(PercentageOfOriginal) : TEXT("0")));
 	}
 	else
 	{
diff --git a/Engine/Source/Editor/UnrealEd/Private/HierarchicalLOD.cpp b/Engine/Source/Editor/UnrealEd/Private/HierarchicalLOD.cpp
--- a/Engine/Source/Editor/UnrealEd/Private/HierarchicalLOD.cpp
+++ b/Engine/Source/Editor/UnrealEd/Private/HierarchicalLOD.cpp
@@ -779,7 +779,8 @@ void FHierarchicalLODBuilder::BuildMeshesForLODActors(bool bForceAll)
 		}
 
 		FScopedSlowTask SlowTask(105, (LOCTEXT("HierarchicalLOD_BuildLODActorMeshes", "Building LODActor meshes")));
-		SlowTask.MakeDialog();
+		const bool bShowCancelButton = true;
+		SlowTask.MakeDialog(bShowCancelButton);
 
 		const TArray<FHierarchicalSimplification>& BuildLODLevelSettings = LevelIter->GetWorldSettings()->GetHierarchicalLODSetup();
 		UMaterialInterface* BaseMaterial = LevelIter->GetWorldSettings()->GetHierarchicalLODBaseMaterial();
@@ -863,6 +864,9 @@ void FHierarchicalLODBuilder::BuildMeshesForLODActors(bool bForceAll)
 								}
 
 								++LODActorIndex;
+
+								if (SlowTask.ShouldCancel())
+									break;
 							}
 						}
 					}
@@ -875,6 +879,9 @@ void FHierarchicalLODBuilder::BuildMeshesForLODActors(bool bForceAll)
 							Proxy->Clean();
 						}
 					}
+
+					if (SlowTask.ShouldCancel())
+						break;
 				}
 			}
 		}
