Use actor labels rather than actor names in the HLOD Outliner

diff --git a/Engine/Source/Editor/HierarchicalLODOutliner/Private/StaticMeshActorItem.cpp b/Engine/Source/Editor/HierarchicalLODOutliner/Private/StaticMeshActorItem.cpp
--- a/Engine/Source/Editor/HierarchicalLODOutliner/Private/StaticMeshActorItem.cpp
+++ b/Engine/Source/Editor/HierarchicalLODOutliner/Private/StaticMeshActorItem.cpp
@@ -32,7 +32,7 @@ FString HLODOutliner::FStaticMeshActorItem::GetDisplayString() const
 {
 	if (StaticMeshActor.IsValid())
 	{
-		return StaticMeshActor->GetFName().GetPlainNameString();
+		return StaticMeshActor->GetActorNameOrLabel();
 	}
 	else
 	{
